Fixing more cases of FindFunctionByName not properly using GET_FUNCTION_NAME_CHECKED

diff --git a/Engine/Source/Editor/GameplayTasksEditor/Private/K2Node_LatentGameplayTaskCall.cpp b/Engine/Source/Editor/GameplayTasksEditor/Private/K2Node_LatentGameplayTaskCall.cpp
--- a/Engine/Source/Editor/GameplayTasksEditor/Private/K2Node_LatentGameplayTaskCall.cpp
+++ b/Engine/Source/Editor/GameplayTasksEditor/Private/K2Node_LatentGameplayTaskCall.cpp
@@ -736,23 +736,23 @@ void UK2Node_LatentGameplayTaskCall::ExpandNode(class FKismetCompilerContext& Co
 
 		//Increment iterator command
 		UK2Node_CallFunction* Increment = CompilerContext.SpawnIntermediateNode<UK2Node_CallFunction>(this, SourceGraph);
-		Increment->SetFromFunction(UKismetMathLibrary::StaticClass()->FindFunctionByName(TEXT("Add_IntInt")));
+		Increment->SetFromFunction(UKismetMathLibrary::StaticClass()->FindFunctionByName(GET_FUNCTION_NAME_CHECKED(UKismetMathLibrary, Add_IntInt)));
 		Increment->AllocateDefaultPins();
 		Increment->FindPinChecked(TEXT("B"))->DefaultValue = TEXT("1");
 
 		//Array length
 		UK2Node_CallArrayFunction* ArrayLength = CompilerContext.SpawnIntermediateNode<UK2Node_CallArrayFunction>(this, SourceGraph);
-		ArrayLength->SetFromFunction(UKismetArrayLibrary::StaticClass()->FindFunctionByName(TEXT("Array_Length")));
+		ArrayLength->SetFromFunction(UKismetArrayLibrary::StaticClass()->FindFunctionByName(GET_FUNCTION_NAME_CHECKED(UKismetArrayLibrary, Array_Length)));
 		ArrayLength->AllocateDefaultPins();
 
 		//Array element retrieval
 		UK2Node_CallArrayFunction* GetElement = CompilerContext.SpawnIntermediateNode<UK2Node_CallArrayFunction>(this, SourceGraph);
-		GetElement->SetFromFunction(UKismetArrayLibrary::StaticClass()->FindFunctionByName(TEXT("Array_Get")));
+		GetElement->SetFromFunction(UKismetArrayLibrary::StaticClass()->FindFunctionByName(GET_FUNCTION_NAME_CHECKED(UKismetArrayLibrary, Array_Get)));
 		GetElement->AllocateDefaultPins();
 
 		//Check node for iterator versus array length
 		UK2Node_CallFunction* Condition = CompilerContext.SpawnIntermediateNode<UK2Node_CallFunction>(this, SourceGraph);
-		Condition->SetFromFunction(UKismetMathLibrary::StaticClass()->FindFunctionByName(TEXT("Less_IntInt")));
+		Condition->SetFromFunction(UKismetMathLibrary::StaticClass()->FindFunctionByName(GET_FUNCTION_NAME_CHECKED(UKismetMathLibrary, Less_IntInt)));
 		Condition->AllocateDefaultPins();
 
 		//Connections to set up the loop
