Fixes for crashes introduced by optimization in 2674041

diff --git a/Engine/Source/Editor/BlueprintGraph/Private/K2Node_BaseAsyncTask.cpp b/Engine/Source/Editor/BlueprintGraph/Private/K2Node_BaseAsyncTask.cpp
--- a/Engine/Source/Editor/BlueprintGraph/Private/K2Node_BaseAsyncTask.cpp
+++ b/Engine/Source/Editor/BlueprintGraph/Private/K2Node_BaseAsyncTask.cpp
@@ -93,7 +93,7 @@ void UK2Node_BaseAsyncTask::AllocateDefaultPins()
 	}
 
 	bool bAllPinsGood = true;
-	UFunction* Function = ProxyFactoryClass->FindFunctionByName(ProxyFactoryFunctionName);
+	UFunction* Function = ProxyFactoryClass ? ProxyFactoryClass->FindFunctionByName(ProxyFactoryFunctionName) : nullptr;
 	if (Function)
 	{
 		TSet<FString> PinsToHide;
diff --git a/Engine/Source/Editor/GameplayTasksEditor/Private/K2Node_LatentGameplayTaskCall.cpp b/Engine/Source/Editor/GameplayTasksEditor/Private/K2Node_LatentGameplayTaskCall.cpp
--- a/Engine/Source/Editor/GameplayTasksEditor/Private/K2Node_LatentGameplayTaskCall.cpp
+++ b/Engine/Source/Editor/GameplayTasksEditor/Private/K2Node_LatentGameplayTaskCall.cpp
@@ -277,16 +277,16 @@ bool UK2Node_LatentGameplayTaskCall::IsSpawnVarPin(UEdGraphPin* Pin)
 bool UK2Node_LatentGameplayTaskCall::ValidateActorSpawning(class FKismetCompilerContext& CompilerContext, bool bGenerateErrors)
 {
 	FName ProxyPrespawnFunctionName = *FK2Node_LatentAbilityCallHelper::BeginSpawnFuncName;
-	UFunction* PreSpawnFunction = ProxyFactoryClass->FindFunctionByName(ProxyPrespawnFunctionName);
+	UFunction* PreSpawnFunction = ProxyFactoryClass ? ProxyFactoryClass->FindFunctionByName(ProxyPrespawnFunctionName) : nullptr;
 
 	FName ProxyPostpawnFunctionName = *FK2Node_LatentAbilityCallHelper::FinishSpawnFuncName;
-	UFunction* PostSpawnFunction = ProxyFactoryClass->FindFunctionByName(ProxyPostpawnFunctionName);
+	UFunction* PostSpawnFunction = ProxyFactoryClass ? ProxyFactoryClass->FindFunctionByName(ProxyPostpawnFunctionName) : nullptr;
 
 	FName ProxyPrespawnArrayFunctionName = *FK2Node_LatentAbilityCallHelper::BeginSpawnArrayFuncName;
-	UFunction* PreSpawnArrayFunction = ProxyFactoryClass->FindFunctionByName(ProxyPrespawnArrayFunctionName);
+	UFunction* PreSpawnArrayFunction = ProxyFactoryClass ? ProxyFactoryClass->FindFunctionByName(ProxyPrespawnArrayFunctionName) : nullptr;
 
 	FName ProxyPostpawnArrayFunctionName = *FK2Node_LatentAbilityCallHelper::FinishSpawnArrayFuncName;
-	UFunction* PostSpawnArrayFunction = ProxyFactoryClass->FindFunctionByName(ProxyPostpawnArrayFunctionName);
+	UFunction* PostSpawnArrayFunction = ProxyFactoryClass ? ProxyFactoryClass->FindFunctionByName(ProxyPostpawnArrayFunctionName) : nullptr;
 
 	bool HasClassParameter = GetClassPin() != nullptr;
 	bool HasPreSpawnFunc = PreSpawnFunction != nullptr;
@@ -337,16 +337,16 @@ bool UK2Node_LatentGameplayTaskCall::ValidateActorSpawning(class FKismetCompiler
 bool UK2Node_LatentGameplayTaskCall::ValidateActorArraySpawning(class FKismetCompilerContext& CompilerContext, bool bGenerateErrors)
 {
 	FName ProxyPrespawnFunctionName = *FK2Node_LatentAbilityCallHelper::BeginSpawnFuncName;
-	UFunction* PreSpawnFunction = ProxyFactoryClass->FindFunctionByName(ProxyPrespawnFunctionName);
+	UFunction* PreSpawnFunction = ProxyFactoryClass ? ProxyFactoryClass->FindFunctionByName(ProxyPrespawnFunctionName) : nullptr;
 
 	FName ProxyPostpawnFunctionName = *FK2Node_LatentAbilityCallHelper::FinishSpawnFuncName;
-	UFunction* PostSpawnFunction = ProxyFactoryClass->FindFunctionByName(ProxyPostpawnFunctionName);
+	UFunction* PostSpawnFunction = ProxyFactoryClass ? ProxyFactoryClass->FindFunctionByName(ProxyPostpawnFunctionName) : nullptr;
 
 	FName ProxyPrespawnArrayFunctionName = *FK2Node_LatentAbilityCallHelper::BeginSpawnArrayFuncName;
-	UFunction* PreSpawnArrayFunction = ProxyFactoryClass->FindFunctionByName(ProxyPrespawnArrayFunctionName);
+	UFunction* PreSpawnArrayFunction = ProxyFactoryClass ? ProxyFactoryClass->FindFunctionByName(ProxyPrespawnArrayFunctionName) : nullptr;
 
 	FName ProxyPostpawnArrayFunctionName = *FK2Node_LatentAbilityCallHelper::FinishSpawnArrayFuncName;
-	UFunction* PostSpawnArrayFunction = ProxyFactoryClass->FindFunctionByName(ProxyPostpawnArrayFunctionName);
+	UFunction* PostSpawnArrayFunction = ProxyFactoryClass ? ProxyFactoryClass->FindFunctionByName(ProxyPostpawnArrayFunctionName) : nullptr;
 
 	bool HasClassParameter = GetClassToSpawn() != nullptr;
 	bool HasPreSpawnFunc = PreSpawnFunction != nullptr;
@@ -593,10 +593,10 @@ void UK2Node_LatentGameplayTaskCall::ExpandNode(class FKismetCompilerContext& Co
 	// ------------------------------------------------------------------------------------------
 
 	FName ProxyPrespawnFunctionName = bValidatedActorArraySpawn ? *FK2Node_LatentAbilityCallHelper::BeginSpawnArrayFuncName : *FK2Node_LatentAbilityCallHelper::BeginSpawnFuncName;
-	UFunction* PreSpawnFunction = ProxyFactoryClass->FindFunctionByName(ProxyPrespawnFunctionName);
+	UFunction* PreSpawnFunction = ProxyFactoryClass ? ProxyFactoryClass->FindFunctionByName(ProxyPrespawnFunctionName) : nullptr;
 
 	FName ProxyPostpawnFunctionName = bValidatedActorArraySpawn ? *FK2Node_LatentAbilityCallHelper::FinishSpawnArrayFuncName : *FK2Node_LatentAbilityCallHelper::FinishSpawnFuncName;
-	UFunction* PostSpawnFunction = ProxyFactoryClass->FindFunctionByName(ProxyPostpawnFunctionName);
+	UFunction* PostSpawnFunction = ProxyFactoryClass ? ProxyFactoryClass->FindFunctionByName(ProxyPostpawnFunctionName) : nullptr;
 
 	if (PreSpawnFunction == nullptr)
 	{
