Fix for Blueprint property incorrectly marked as transient.
Fix for Renderer crash in editor
#tests Golden path in Paragon

diff --git a/Engine/Source/Editor/GameplayTasksEditor/Classes/K2Node_LatentGameplayTaskCall.h b/Engine/Source/Editor/GameplayTasksEditor/Classes/K2Node_LatentGameplayTaskCall.h
--- a/Engine/Source/Editor/GameplayTasksEditor/Classes/K2Node_LatentGameplayTaskCall.h
+++ b/Engine/Source/Editor/GameplayTasksEditor/Classes/K2Node_LatentGameplayTaskCall.h
@@ -32,8 +32,8 @@ public:
 	bool ValidateActorSpawning(class FKismetCompilerContext& CompilerContext, bool bGenerateErrors);
 	bool ValidateActorArraySpawning(class FKismetCompilerContext& CompilerContext, bool bGenerateErrors);
 
-	UPROPERTY(Transient)
-	TArray<FString> SpawnParmPins;
+	UPROPERTY()
+	TArray<FString> SpawnParamPins;
 
 	static void RegisterSpecializedTaskNodeClass(TSubclassOf<UK2Node_LatentGameplayTaskCall> NodeClass);
 protected:
diff --git a/Engine/Source/Editor/GameplayTasksEditor/Private/K2Node_LatentGameplayTaskCall.cpp b/Engine/Source/Editor/GameplayTasksEditor/Private/K2Node_LatentGameplayTaskCall.cpp
--- a/Engine/Source/Editor/GameplayTasksEditor/Private/K2Node_LatentGameplayTaskCall.cpp
+++ b/Engine/Source/Editor/GameplayTasksEditor/Private/K2Node_LatentGameplayTaskCall.cpp
@@ -170,7 +170,7 @@ void UK2Node_LatentGameplayTaskCall::CreatePinsForClass(UClass* InClass)
 
 	const UObject* const ClassDefaultObject = InClass->GetDefaultObject(false);
 
-	SpawnParmPins.Empty();
+	SpawnParamPins.Empty();
 
 	// Tasks can hide spawn parameters by doing meta = (HideSpawnParms="PropertyA,PropertyB")
 	// (For example, hide Instigator in situations where instigator is not relevant to your task)
@@ -207,7 +207,7 @@ void UK2Node_LatentGameplayTaskCall::CreatePinsForClass(UClass* InClass)
 
 			UEdGraphPin* Pin = CreatePin(EGPD_Input, TEXT(""), TEXT(""), NULL, false, false, Property->GetName());
 			const bool bPinGood = (Pin != NULL) && K2Schema->ConvertPropertyToPinType(Property, /*out*/ Pin->PinType);
-			SpawnParmPins.Add(Pin->PinName);
+			SpawnParamPins.Add(Pin->PinName);
 
 			if (ClassDefaultObject && Pin && K2Schema->PinDefaultValueIsEditable(*Pin))
 			{
@@ -237,7 +237,7 @@ void UK2Node_LatentGameplayTaskCall::PinDefaultValueChanged(UEdGraphPin* Changed
 		//ResultPin->BreakAllPinLinks();
 
 		// Remove all pins related to archetype variables
-		for (const FString& OldPinReference : SpawnParmPins)
+		for (const FString& OldPinReference : SpawnParamPins)
 		{
 			UEdGraphPin* OldPin = FindPin(OldPinReference);
 			if(OldPin)
@@ -246,7 +246,7 @@ void UK2Node_LatentGameplayTaskCall::PinDefaultValueChanged(UEdGraphPin* Changed
 				Pins.Remove(OldPin);
 			}
 		}
-		SpawnParmPins.Empty();
+		SpawnParamPins.Empty();
 
 		UClass* UseSpawnClass = GetClassToSpawn();
 		if (UseSpawnClass != NULL)
@@ -408,7 +408,7 @@ bool UK2Node_LatentGameplayTaskCall::ValidateActorArraySpawning(class FKismetCom
 bool UK2Node_LatentGameplayTaskCall::ConnectSpawnProperties(UClass* ClassToSpawn, const UEdGraphSchema_K2* Schema, class FKismetCompilerContext& CompilerContext, UEdGraph* SourceGraph, UEdGraphPin*& LastThenPin, UEdGraphPin* SpawnedActorReturnPin)
 {
 	bool bIsErrorFree = true;
-	for (const FString& OldPinReference : SpawnParmPins)
+	for (const FString& OldPinReference : SpawnParamPins)
 	{
 		UEdGraphPin* SpawnVarPin = FindPin(OldPinReference);
 		const bool bHasDefaultValue = SpawnVarPin && (!SpawnVarPin->DefaultValue.IsEmpty() || !SpawnVarPin->DefaultTextValue.IsEmpty() || SpawnVarPin->DefaultObject);
diff --git a/Engine/Source/Runtime/Renderer/Private/PostProcess/PostProcessTonemap.cpp b/Engine/Source/Runtime/Renderer/Private/PostProcess/PostProcessTonemap.cpp
--- a/Engine/Source/Runtime/Renderer/Private/PostProcess/PostProcessTonemap.cpp
+++ b/Engine/Source/Runtime/Renderer/Private/PostProcess/PostProcessTonemap.cpp
@@ -1342,6 +1342,7 @@ void FRCPassPostProcessTonemap::Process(FRenderingCompositePassContext& Context)
 	case 11: SetShaderTempl<11>(Context, bDoEyeAdaptation); break;
 	case 12: SetShaderTempl<12>(Context, bDoEyeAdaptation); break;
 	case 13: SetShaderTempl<13>(Context, bDoEyeAdaptation); break;
+	case 14: SetShaderTempl<14>(Context, bDoEyeAdaptation); break;
 	default:
 		check(0);
 	}
