Copying //UE4/Dev-Framework to //UE4/Dev-Main (Source: //UE4/Dev-Framework @ 3716594)

diff --git a/Engine/Source/Editor/GameplayTasksEditor/Private/K2Node_LatentGameplayTaskCall.cpp b/Engine/Source/Editor/GameplayTasksEditor/Private/K2Node_LatentGameplayTaskCall.cpp
--- a/Engine/Source/Editor/GameplayTasksEditor/Private/K2Node_LatentGameplayTaskCall.cpp
+++ b/Engine/Source/Editor/GameplayTasksEditor/Private/K2Node_LatentGameplayTaskCall.cpp
@@ -38,7 +38,7 @@ void UK2Node_LatentGameplayTaskCall::RegisterSpecializedTaskNodeClass(TSubclassO
 
 bool UK2Node_LatentGameplayTaskCall::HasDedicatedNodeClass(TSubclassOf<UGameplayTask> TaskClass) 
 {
-	for (const auto& NodeClass : NodeClasses)
+	for (const TWeakObjectPtr<UClass>& NodeClass : NodeClasses)
 	{
 		if (NodeClass.IsValid())
 		{
@@ -101,22 +101,22 @@ void UK2Node_LatentGameplayTaskCall::GetMenuActions(FBlueprintActionDatabaseRegi
 
 struct FK2Node_LatentAbilityCallHelper
 {
-	static FString WorldContextPinName;
-	static FString ClassPinName;
-	static FString BeginSpawnFuncName;
-	static FString FinishSpawnFuncName;
-	static FString BeginSpawnArrayFuncName;
-	static FString FinishSpawnArrayFuncName;
-	static FString SpawnedActorPinName;
+	static FName WorldContextPinName;
+	static FName ClassPinName;
+	static FName BeginSpawnFuncName;
+	static FName FinishSpawnFuncName;
+	static FName BeginSpawnArrayFuncName;
+	static FName FinishSpawnArrayFuncName;
+	static FName SpawnedActorPinName;
 };
 
-FString FK2Node_LatentAbilityCallHelper::WorldContextPinName(TEXT("WorldContextObject"));
-FString FK2Node_LatentAbilityCallHelper::ClassPinName(TEXT("Class"));
-FString FK2Node_LatentAbilityCallHelper::BeginSpawnFuncName(TEXT("BeginSpawningActor"));
-FString FK2Node_LatentAbilityCallHelper::FinishSpawnFuncName(TEXT("FinishSpawningActor"));
-FString FK2Node_LatentAbilityCallHelper::BeginSpawnArrayFuncName(TEXT("BeginSpawningActorArray"));
-FString FK2Node_LatentAbilityCallHelper::FinishSpawnArrayFuncName(TEXT("FinishSpawningActorArray"));
-FString FK2Node_LatentAbilityCallHelper::SpawnedActorPinName(TEXT("SpawnedActor"));
+FName FK2Node_LatentAbilityCallHelper::WorldContextPinName(TEXT("WorldContextObject"));
+FName FK2Node_LatentAbilityCallHelper::ClassPinName(TEXT("Class"));
+FName FK2Node_LatentAbilityCallHelper::BeginSpawnFuncName(TEXT("BeginSpawningActor"));
+FName FK2Node_LatentAbilityCallHelper::FinishSpawnFuncName(TEXT("FinishSpawningActor"));
+FName FK2Node_LatentAbilityCallHelper::BeginSpawnArrayFuncName(TEXT("BeginSpawningActorArray"));
+FName FK2Node_LatentAbilityCallHelper::FinishSpawnArrayFuncName(TEXT("FinishSpawningActorArray"));
+FName FK2Node_LatentAbilityCallHelper::SpawnedActorPinName(TEXT("SpawnedActor"));
 
 // -------------------------------------------------
 
@@ -135,17 +135,16 @@ UEdGraphPin* UK2Node_LatentGameplayTaskCall::GetClassPin(const TArray<UEdGraphPi
 {
 	const TArray<UEdGraphPin*>* PinsToSearch = InPinsToSearch ? InPinsToSearch : &Pins;
 
-	UEdGraphPin* Pin = NULL;
-	for (auto PinIt = PinsToSearch->CreateConstIterator(); PinIt; ++PinIt)
+	UEdGraphPin* Pin = nullptr;
+	for (UEdGraphPin* TestPin : *PinsToSearch)
 	{
-		UEdGraphPin* TestPin = *PinIt;
 		if (TestPin && TestPin->PinName == FK2Node_LatentAbilityCallHelper::ClassPinName)
 		{
 			Pin = TestPin;
 			break;
 		}
 	}
-	check(Pin == NULL || Pin->Direction == EGPD_Input);
+	check(Pin == nullptr || Pin->Direction == EGPD_Input);
 	return Pin;
 }
 
@@ -161,7 +160,7 @@ UClass* UK2Node_LatentGameplayTaskCall::GetClassToSpawn(const TArray<UEdGraphPin
 	}
 	else if (ClassPin && (1 == ClassPin->LinkedTo.Num()))
 	{
-		auto SourcePin = ClassPin->LinkedTo[0];
+		UEdGraphPin* SourcePin = ClassPin->LinkedTo[0];
 		UseSpawnClass = SourcePin ? Cast<UClass>(SourcePin->PinType.PinSubCategoryObject.Get()) : nullptr;
 	}
 
@@ -176,7 +175,7 @@ void UK2Node_LatentGameplayTaskCall::CreatePinsForClass(UClass* InClass)
 
 	const UObject* const ClassDefaultObject = InClass->GetDefaultObject(false);
 
-	SpawnParamPins.Empty();
+	SpawnParamPins.Reset();
 
 	// Tasks can hide spawn parameters by doing meta = (HideSpawnParms="PropertyA,PropertyB")
 	// (For example, hide Instigator in situations where instigator is not relevant to your task)
@@ -207,11 +206,11 @@ void UK2Node_LatentGameplayTaskCall::CreatePinsForClass(UClass* InClass)
 			Property->HasAllPropertyFlags(CPF_BlueprintVisible) &&
 			!bIsDelegate && 
 			!IgnorePropertyList.Contains(Property->GetName()) &&
-			(FindPin(Property->GetName()) == nullptr) )
+			(FindPin(Property->GetFName()) == nullptr) )
 		{
 
 
-			UEdGraphPin* Pin = CreatePin(EGPD_Input, FString(), FString(), nullptr, Property->GetName());
+			UEdGraphPin* Pin = CreatePin(EGPD_Input, NAME_None, Property->GetFName());
 			check(Pin);
 			const bool bPinGood = K2Schema->ConvertPropertyToPinType(Property, /*out*/ Pin->PinType);
 			SpawnParamPins.Add(Pin->PinName);
@@ -241,7 +240,7 @@ void UK2Node_LatentGameplayTaskCall::PinDefaultValueChanged(UEdGraphPin* Changed
 		//ResultPin->BreakAllPinLinks();
 
 		// Remove all pins related to archetype variables
-		for (const FString& OldPinReference : SpawnParamPins)
+		for (const FName& OldPinReference : SpawnParamPins)
 		{
 			UEdGraphPin* OldPin = FindPin(OldPinReference);
 			if(OldPin)
@@ -250,7 +249,7 @@ void UK2Node_LatentGameplayTaskCall::PinDefaultValueChanged(UEdGraphPin* Changed
 				Pins.Remove(OldPin);
 			}
 		}
-		SpawnParamPins.Empty();
+		SpawnParamPins.Reset();
 
 		UClass* UseSpawnClass = GetClassToSpawn();
 		if (UseSpawnClass != NULL)
@@ -269,21 +268,17 @@ void UK2Node_LatentGameplayTaskCall::PinDefaultValueChanged(UEdGraphPin* Changed
 
 UEdGraphPin* UK2Node_LatentGameplayTaskCall::GetResultPin() const
 {
-	const UEdGraphSchema_K2* K2Schema = GetDefault<UEdGraphSchema_K2>();
-
-	UEdGraphPin* Pin = FindPinChecked(K2Schema->PN_ReturnValue);
+	UEdGraphPin* Pin = FindPinChecked(UEdGraphSchema_K2::PN_ReturnValue);
 	check(Pin->Direction == EGPD_Output);
 	return Pin;
 }
 
 bool UK2Node_LatentGameplayTaskCall::IsSpawnVarPin(UEdGraphPin* Pin)
 {
-	const UEdGraphSchema_K2* K2Schema = GetDefault<UEdGraphSchema_K2>();
-
 	return (Pin->Direction == EEdGraphPinDirection::EGPD_Input &&
-			Pin->PinName != K2Schema->PN_Execute &&
-			Pin->PinName != K2Schema->PN_Then &&
-			Pin->PinName != K2Schema->PN_ReturnValue &&
+			Pin->PinName != UEdGraphSchema_K2::PN_Execute &&
+			Pin->PinName != UEdGraphSchema_K2::PN_Then &&
+			Pin->PinName != UEdGraphSchema_K2::PN_ReturnValue &&
 			Pin->PinName != FK2Node_LatentAbilityCallHelper::ClassPinName &&
 			Pin->PinName != FK2Node_LatentAbilityCallHelper::WorldContextPinName);
 
@@ -291,16 +286,16 @@ bool UK2Node_LatentGameplayTaskCall::IsSpawnVarPin(UEdGraphPin* Pin)
 
 bool UK2Node_LatentGameplayTaskCall::ValidateActorSpawning(class FKismetCompilerContext& CompilerContext, bool bGenerateErrors)
 {
-	FName ProxyPrespawnFunctionName = *FK2Node_LatentAbilityCallHelper::BeginSpawnFuncName;
+	FName ProxyPrespawnFunctionName = FK2Node_LatentAbilityCallHelper::BeginSpawnFuncName;
 	UFunction* PreSpawnFunction = ProxyFactoryClass ? ProxyFactoryClass->FindFunctionByName(ProxyPrespawnFunctionName) : nullptr;
 
-	FName ProxyPostpawnFunctionName = *FK2Node_LatentAbilityCallHelper::FinishSpawnFuncName;
+	FName ProxyPostpawnFunctionName = FK2Node_LatentAbilityCallHelper::FinishSpawnFuncName;
 	UFunction* PostSpawnFunction = ProxyFactoryClass ? ProxyFactoryClass->FindFunctionByName(ProxyPostpawnFunctionName) : nullptr;
 
-	FName ProxyPrespawnArrayFunctionName = *FK2Node_LatentAbilityCallHelper::BeginSpawnArrayFuncName;
+	FName ProxyPrespawnArrayFunctionName = FK2Node_LatentAbilityCallHelper::BeginSpawnArrayFuncName;
 	UFunction* PreSpawnArrayFunction = ProxyFactoryClass ? ProxyFactoryClass->FindFunctionByName(ProxyPrespawnArrayFunctionName) : nullptr;
 
-	FName ProxyPostpawnArrayFunctionName = *FK2Node_LatentAbilityCallHelper::FinishSpawnArrayFuncName;
+	FName ProxyPostpawnArrayFunctionName = FK2Node_LatentAbilityCallHelper::FinishSpawnArrayFuncName;
 	UFunction* PostSpawnArrayFunction = ProxyFactoryClass ? ProxyFactoryClass->FindFunctionByName(ProxyPostpawnArrayFunctionName) : nullptr;
 
 	bool HasClassParameter = GetClassPin() != nullptr;
@@ -351,16 +346,16 @@ bool UK2Node_LatentGameplayTaskCall::ValidateActorSpawning(class FKismetCompiler
 
 bool UK2Node_LatentGameplayTaskCall::ValidateActorArraySpawning(class FKismetCompilerContext& CompilerContext, bool bGenerateErrors)
 {
-	FName ProxyPrespawnFunctionName = *FK2Node_LatentAbilityCallHelper::BeginSpawnFuncName;
+	FName ProxyPrespawnFunctionName = FK2Node_LatentAbilityCallHelper::BeginSpawnFuncName;
 	UFunction* PreSpawnFunction = ProxyFactoryClass ? ProxyFactoryClass->FindFunctionByName(ProxyPrespawnFunctionName) : nullptr;
 
-	FName ProxyPostpawnFunctionName = *FK2Node_LatentAbilityCallHelper::FinishSpawnFuncName;
+	FName ProxyPostpawnFunctionName = FK2Node_LatentAbilityCallHelper::FinishSpawnFuncName;
 	UFunction* PostSpawnFunction = ProxyFactoryClass ? ProxyFactoryClass->FindFunctionByName(ProxyPostpawnFunctionName) : nullptr;
 
-	FName ProxyPrespawnArrayFunctionName = *FK2Node_LatentAbilityCallHelper::BeginSpawnArrayFuncName;
+	FName ProxyPrespawnArrayFunctionName = FK2Node_LatentAbilityCallHelper::BeginSpawnArrayFuncName;
 	UFunction* PreSpawnArrayFunction = ProxyFactoryClass ? ProxyFactoryClass->FindFunctionByName(ProxyPrespawnArrayFunctionName) : nullptr;
 
-	FName ProxyPostpawnArrayFunctionName = *FK2Node_LatentAbilityCallHelper::FinishSpawnArrayFuncName;
+	FName ProxyPostpawnArrayFunctionName = FK2Node_LatentAbilityCallHelper::FinishSpawnArrayFuncName;
 	UFunction* PostSpawnArrayFunction = ProxyFactoryClass ? ProxyFactoryClass->FindFunctionByName(ProxyPostpawnArrayFunctionName) : nullptr;
 
 	bool HasClassParameter = GetClassToSpawn() != nullptr;
@@ -412,7 +407,7 @@ bool UK2Node_LatentGameplayTaskCall::ValidateActorArraySpawning(class FKismetCom
 bool UK2Node_LatentGameplayTaskCall::ConnectSpawnProperties(UClass* ClassToSpawn, const UEdGraphSchema_K2* Schema, class FKismetCompilerContext& CompilerContext, UEdGraph* SourceGraph, UEdGraphPin*& LastThenPin, UEdGraphPin* SpawnedActorReturnPin)
 {
 	bool bIsErrorFree = true;
-	for (const FString& OldPinReference : SpawnParamPins)
+	for (const FName& OldPinReference : SpawnParamPins)
 	{
 		UEdGraphPin* SpawnVarPin = FindPin(OldPinReference);
 		if (!SpawnVarPin)
@@ -425,7 +420,7 @@ bool UK2Node_LatentGameplayTaskCall::ConnectSpawnProperties(UClass* ClassToSpawn
 		{
 			if (SpawnVarPin->LinkedTo.Num() == 0)
 			{
-				UProperty* Property = FindField<UProperty>(ClassToSpawn, *SpawnVarPin->PinName);
+				UProperty* Property = FindField<UProperty>(ClassToSpawn, SpawnVarPin->PinName);
 				// NULL property indicates that this pin was part of the original node, not the 
 				// class we're assigning to:
 				if (!Property)
@@ -467,9 +462,9 @@ bool UK2Node_LatentGameplayTaskCall::ConnectSpawnProperties(UClass* ClassToSpawn
 				bIsErrorFree &= Schema->TryCreateConnection(LastThenPin, SetVarNode->GetExecPin());
 				LastThenPin = SetVarNode->GetThenPin();
 
-				static FString ObjectParamName = FString(TEXT("Object"));
-				static FString ValueParamName = FString(TEXT("Value"));
-				static FString PropertyNameParamName = FString(TEXT("PropertyName"));
+				static const FName ObjectParamName(TEXT("Object"));
+				static const FName ValueParamName(TEXT("Value"));
+				static const FName PropertyNameParamName(TEXT("PropertyName"));
 
 				// Connect the new actor to the 'object' pin
 				UEdGraphPin* ObjectPin = SetVarNode->FindPinChecked(ObjectParamName);
@@ -477,12 +472,12 @@ bool UK2Node_LatentGameplayTaskCall::ConnectSpawnProperties(UClass* ClassToSpawn
 
 				// Fill in literal for 'property name' pin - name of pin is property name
 				UEdGraphPin* PropertyNamePin = SetVarNode->FindPinChecked(PropertyNameParamName);
-				PropertyNamePin->DefaultValue = SpawnVarPin->PinName;
+				PropertyNamePin->DefaultValue = SpawnVarPin->PinName.ToString();
 
 				UEdGraphPin* ValuePin = SetVarNode->FindPinChecked(ValueParamName);
 				if (SpawnVarPin->LinkedTo.Num() == 0 &&
 					SpawnVarPin->DefaultValue != FString() &&
-					SpawnVarPin->PinType.PinCategory == Schema->PC_Byte &&
+					SpawnVarPin->PinType.PinCategory == UEdGraphSchema_K2::PC_Byte &&
 					SpawnVarPin->PinType.PinSubCategoryObject.IsValid() &&
 					SpawnVarPin->PinType.PinSubCategoryObject->IsA<UEnum>())
 				{
@@ -490,7 +485,7 @@ bool UK2Node_LatentGameplayTaskCall::ConnectSpawnProperties(UClass* ClassToSpawn
 					UK2Node_EnumLiteral* EnumLiteralNode = CompilerContext.SpawnIntermediateNode<UK2Node_EnumLiteral>(this, SourceGraph);
 					EnumLiteralNode->Enum = CastChecked<UEnum>(SpawnVarPin->PinType.PinSubCategoryObject.Get());
 					EnumLiteralNode->AllocateDefaultPins();
-					EnumLiteralNode->FindPinChecked(Schema->PN_ReturnValue)->MakeLinkTo(ValuePin);
+					EnumLiteralNode->FindPinChecked(UEdGraphSchema_K2::PN_ReturnValue)->MakeLinkTo(ValuePin);
 
 					UEdGraphPin* InPin = EnumLiteralNode->FindPinChecked(UK2Node_EnumLiteral::GetEnumInputPinName());
 					InPin->DefaultValue = SpawnVarPin->DefaultValue;
@@ -498,7 +493,7 @@ bool UK2Node_LatentGameplayTaskCall::ConnectSpawnProperties(UClass* ClassToSpawn
 				else
 				{
 					// For non-array struct pins that are not linked, transfer the pin type so that the node will expand an auto-ref that will assign the value by-ref.
-					if (SpawnVarPin->PinType.IsArray() == false && SpawnVarPin->PinType.PinCategory == Schema->PC_Struct && SpawnVarPin->LinkedTo.Num() == 0)
+					if (SpawnVarPin->PinType.IsArray() == false && SpawnVarPin->PinType.PinCategory == UEdGraphSchema_K2::PC_Struct && SpawnVarPin->LinkedTo.Num() == 0)
 					{
 						ValuePin->PinType.PinCategory = SpawnVarPin->PinType.PinCategory;
 						ValuePin->PinType.PinSubCategory = SpawnVarPin->PinType.PinSubCategory;
@@ -568,10 +563,10 @@ void UK2Node_LatentGameplayTaskCall::ExpandNode(class FKismetCompilerContext& Co
 	UK2Node_CallFunction* const CallCreateProxyObjectNode = CompilerContext.SpawnIntermediateNode<UK2Node_CallFunction>(this, SourceGraph);
 	CallCreateProxyObjectNode->FunctionReference.SetExternalMember(ProxyFactoryFunctionName, ProxyFactoryClass);
 	CallCreateProxyObjectNode->AllocateDefaultPins();
-	bIsErrorFree &= CompilerContext.MovePinLinksToIntermediate(*FindPinChecked(Schema->PN_Execute), *CallCreateProxyObjectNode->FindPinChecked(Schema->PN_Execute)).CanSafeConnect();
-	for (auto CurrentPin : Pins)
+	bIsErrorFree &= CompilerContext.MovePinLinksToIntermediate(*FindPinChecked(UEdGraphSchema_K2::PN_Execute), *CallCreateProxyObjectNode->FindPinChecked(UEdGraphSchema_K2::PN_Execute)).CanSafeConnect();
+	for (UEdGraphPin* CurrentPin : Pins)
 	{
-		if (FBaseAsyncTaskHelper::ValidDataPin(CurrentPin, EGPD_Input, Schema))
+		if (FBaseAsyncTaskHelper::ValidDataPin(CurrentPin, EGPD_Input))
 		{
 			UEdGraphPin* DestPin = CallCreateProxyObjectNode->FindPin(CurrentPin->PinName); // match function inputs, to pass data to function from CallFunction node
 
@@ -586,16 +581,16 @@ void UK2Node_LatentGameplayTaskCall::ExpandNode(class FKismetCompilerContext& Co
 	// Expose Async Task Proxy object
 	UEdGraphPin* const ProxyObjectPin = CallCreateProxyObjectNode->GetReturnValuePin();
 	check(ProxyObjectPin);
-	auto OutputAsyncTaskProxy = FindPinChecked(FBaseAsyncTaskHelper::GetAsyncTaskProxyName());
+	UEdGraphPin* OutputAsyncTaskProxy = FindPinChecked(FBaseAsyncTaskHelper::GetAsyncTaskProxyName());
 	bIsErrorFree &= CompilerContext.MovePinLinksToIntermediate(*OutputAsyncTaskProxy, *ProxyObjectPin).CanSafeConnect();
 
 	// ------------------------------------------------------------------------------------------
 	// GATHER OUTPUT PARAMETERS AND PAIR THEM WITH LOCAL VARIABLES
 	// ------------------------------------------------------------------------------------------
 	TArray<FBaseAsyncTaskHelper::FOutputPinAndLocalVariable> VariableOutputs;
-	for (auto CurrentPin : Pins)
+	for (UEdGraphPin* CurrentPin : Pins)
 	{
-		if ((OutputAsyncTaskProxy != CurrentPin) && FBaseAsyncTaskHelper::ValidDataPin(CurrentPin, EGPD_Output, Schema))
+		if ((OutputAsyncTaskProxy != CurrentPin) && FBaseAsyncTaskHelper::ValidDataPin(CurrentPin, EGPD_Output))
 		{
 			const FEdGraphPinType& PinType = CurrentPin->PinType;
 			UK2Node_TemporaryVariable* TempVarOutput = CompilerContext.SpawnInternalVariable(
@@ -608,7 +603,7 @@ void UK2Node_LatentGameplayTaskCall::ExpandNode(class FKismetCompilerContext& Co
 	// ------------------------------------------------------------------------------------------
 	// FOR EACH DELEGATE DEFINE EVENT, CONNECT IT TO DELEGATE AND IMPLEMENT A CHAIN OF ASSIGMENTS
 	// ------------------------------------------------------------------------------------------
-	UEdGraphPin* LastThenPin = CallCreateProxyObjectNode->FindPinChecked(Schema->PN_Then);
+	UEdGraphPin* LastThenPin = CallCreateProxyObjectNode->FindPinChecked(UEdGraphSchema_K2::PN_Then);
 
 	UK2Node_CallFunction* IsValidFuncNode = CompilerContext.SpawnIntermediateNode<UK2Node_CallFunction>(this, SourceGraph);
 	const FName IsValidFuncName = GET_FUNCTION_NAME_CHECKED(UKismetSystemLibrary, IsValid);
@@ -630,7 +625,7 @@ void UK2Node_LatentGameplayTaskCall::ExpandNode(class FKismetCompilerContext& Co
 		bIsErrorFree &= FBaseAsyncTaskHelper::HandleDelegateImplementation(*PropertyIt, VariableOutputs, ProxyObjectPin, LastThenPin, this, SourceGraph, CompilerContext);
 	}
 
-	if (CallCreateProxyObjectNode->FindPinChecked(Schema->PN_Then) == LastThenPin)
+	if (CallCreateProxyObjectNode->FindPinChecked(UEdGraphSchema_K2::PN_Then) == LastThenPin)
 	{
 		CompilerContext.MessageLog.Error(*LOCTEXT("MissingDelegateProperties", "BaseAsyncTask: Proxy has no delegates defined. @@").ToString(), this);
 		return;
@@ -641,10 +636,10 @@ void UK2Node_LatentGameplayTaskCall::ExpandNode(class FKismetCompilerContext& Co
 	// NEW: CREATE A CALL TO THE PRESPAWN FUNCTION, IF IT RETURNS TRUE, THEN WE WILL SPAWN THE NEW ACTOR
 	// ------------------------------------------------------------------------------------------
 
-	FName ProxyPrespawnFunctionName = bValidatedActorArraySpawn ? *FK2Node_LatentAbilityCallHelper::BeginSpawnArrayFuncName : *FK2Node_LatentAbilityCallHelper::BeginSpawnFuncName;
+	FName ProxyPrespawnFunctionName = bValidatedActorArraySpawn ? FK2Node_LatentAbilityCallHelper::BeginSpawnArrayFuncName : FK2Node_LatentAbilityCallHelper::BeginSpawnFuncName;
 	UFunction* PreSpawnFunction = ProxyFactoryClass ? ProxyFactoryClass->FindFunctionByName(ProxyPrespawnFunctionName) : nullptr;
 
-	FName ProxyPostpawnFunctionName = bValidatedActorArraySpawn ? *FK2Node_LatentAbilityCallHelper::FinishSpawnArrayFuncName : *FK2Node_LatentAbilityCallHelper::FinishSpawnFuncName;
+	FName ProxyPostpawnFunctionName = bValidatedActorArraySpawn ? FK2Node_LatentAbilityCallHelper::FinishSpawnArrayFuncName : FK2Node_LatentAbilityCallHelper::FinishSpawnFuncName;
 	UFunction* PostSpawnFunction = ProxyFactoryClass ? ProxyFactoryClass->FindFunctionByName(ProxyPostpawnFunctionName) : nullptr;
 
 	if (PreSpawnFunction == nullptr)
@@ -685,9 +680,9 @@ void UK2Node_LatentGameplayTaskCall::ExpandNode(class FKismetCompilerContext& Co
 	bIsErrorFree &= Schema->TryCreateConnection(ProxyObjectPin, PrespawnCallSelfPin);
 
 	// Hook up input parameters to PreSpawn
-	for (auto CurrentPin : Pins)
+	for (UEdGraphPin* CurrentPin : Pins)
 	{
-		if (FBaseAsyncTaskHelper::ValidDataPin(CurrentPin, EGPD_Input, Schema))
+		if (FBaseAsyncTaskHelper::ValidDataPin(CurrentPin, EGPD_Input))
 		{
 			UEdGraphPin* DestPin = CallPrespawnProxyObjectNode->FindPin(CurrentPin->PinName);
 			if (DestPin)
@@ -698,9 +693,9 @@ void UK2Node_LatentGameplayTaskCall::ExpandNode(class FKismetCompilerContext& Co
 	}		
 
 	// Hook the activate node up in the exec chain
-	UEdGraphPin* PrespawnExecPin = CallPrespawnProxyObjectNode->FindPinChecked(Schema->PN_Execute);
-	UEdGraphPin* PrespawnThenPin = CallPrespawnProxyObjectNode->FindPinChecked(Schema->PN_Then);
-	UEdGraphPin* PrespawnReturnPin = CallPrespawnProxyObjectNode->FindPinChecked(Schema->PN_ReturnValue);
+	UEdGraphPin* PrespawnExecPin = CallPrespawnProxyObjectNode->FindPinChecked(UEdGraphSchema_K2::PN_Execute);
+	UEdGraphPin* PrespawnThenPin = CallPrespawnProxyObjectNode->FindPinChecked(UEdGraphSchema_K2::PN_Then);
+	UEdGraphPin* PrespawnReturnPin = CallPrespawnProxyObjectNode->FindPinChecked(UEdGraphSchema_K2::PN_ReturnValue);
 	UEdGraphPin* SpawnedActorReturnPin = CallPrespawnProxyObjectNode->FindPinChecked(FK2Node_LatentAbilityCallHelper::SpawnedActorPinName);
 
 	bIsErrorFree &= Schema->TryCreateConnection(LastThenPin, PrespawnExecPin);
@@ -734,7 +729,7 @@ void UK2Node_LatentGameplayTaskCall::ExpandNode(class FKismetCompilerContext& Co
 
 		//Create int Iterator
 		UK2Node_TemporaryVariable* IteratorVar = CompilerContext.SpawnIntermediateNode<UK2Node_TemporaryVariable>(this, SourceGraph);
-		IteratorVar->VariableType.PinCategory = Schema->PC_Int;
+		IteratorVar->VariableType.PinCategory = UEdGraphSchema_K2::PC_Int;
 		IteratorVar->AllocateDefaultPins();
 
 		//Iterator assignment (initialization to zero)
@@ -774,7 +769,7 @@ void UK2Node_LatentGameplayTaskCall::ExpandNode(class FKismetCompilerContext& Co
 		bIsErrorFree &= Schema->TryCreateConnection(SpawnedActorReturnPin, ArrayLength->GetTargetArrayPin());
 		bIsErrorFree &= Schema->TryCreateConnection(Condition->GetReturnValuePin(), Branch->GetConditionPin());
 		bIsErrorFree &= Schema->TryCreateConnection(IteratorVar->GetVariablePin(), Condition->FindPinChecked(TEXT("A")));
-		bIsErrorFree &= Schema->TryCreateConnection(ArrayLength->FindPin(Schema->PN_ReturnValue), Condition->FindPinChecked(TEXT("B")));
+		bIsErrorFree &= Schema->TryCreateConnection(ArrayLength->FindPin(UEdGraphSchema_K2::PN_ReturnValue), Condition->FindPinChecked(TEXT("B")));
 
 		//Connections to establish loop iteration
 		bIsErrorFree &= Schema->TryCreateConnection(IteratorVar->GetVariablePin(), Increment->FindPinChecked(TEXT("A")));
@@ -785,8 +780,8 @@ void UK2Node_LatentGameplayTaskCall::ExpandNode(class FKismetCompilerContext& Co
 		//This is the inner loop
 		LastThenPin = Branch->GetThenPin();		//Connect the loop branch to the spawn-assignment code block
 		bIsErrorFree &= Schema->TryCreateConnection(SpawnedActorReturnPin, GetElement->GetTargetArrayPin());
-		bIsErrorFree &= Schema->TryCreateConnection(IteratorVar->GetVariablePin(), GetElement->FindPinChecked(Schema->PN_Index));
-		bIsErrorFree &= ConnectSpawnProperties(ClassToSpawn, Schema, CompilerContext, SourceGraph, LastThenPin, GetElement->FindPinChecked(Schema->PN_Item));		//Last argument is the array element
+		bIsErrorFree &= Schema->TryCreateConnection(IteratorVar->GetVariablePin(), GetElement->FindPinChecked(UEdGraphSchema_K2::PN_Index));
+		bIsErrorFree &= ConnectSpawnProperties(ClassToSpawn, Schema, CompilerContext, SourceGraph, LastThenPin, GetElement->FindPinChecked(UEdGraphSchema_K2::PN_Item));		//Last argument is the array element
 		bIsErrorFree &= Schema->TryCreateConnection(LastThenPin, IteratorAssign->GetExecPin());		//Connect the spawn-assignment code block to the iterator increment
 		
 		//Finish by providing the proper path out
@@ -818,12 +813,12 @@ void UK2Node_LatentGameplayTaskCall::ExpandNode(class FKismetCompilerContext& Co
 	bIsErrorFree &= Schema->TryCreateConnection(ProxyObjectPin, PostspawnCallSelfPin);
 
 	// Link our Postspawn call in
-	bIsErrorFree &= Schema->TryCreateConnection(LastThenPin, CallPostSpawnnProxyObjectNode->FindPinChecked(Schema->PN_Execute));
+	bIsErrorFree &= Schema->TryCreateConnection(LastThenPin, CallPostSpawnnProxyObjectNode->FindPinChecked(UEdGraphSchema_K2::PN_Execute));
 
 	// Hook up any other input parameters to PostSpawn
-	for (auto CurrentPin : Pins)
+	for (UEdGraphPin* CurrentPin : Pins)
 	{
-		if (FBaseAsyncTaskHelper::ValidDataPin(CurrentPin, EGPD_Input, Schema))
+		if (FBaseAsyncTaskHelper::ValidDataPin(CurrentPin, EGPD_Input))
 		{
 			UEdGraphPin* DestPin = CallPostSpawnnProxyObjectNode->FindPin(CurrentPin->PinName);
 			if (DestPin)
@@ -843,11 +838,11 @@ void UK2Node_LatentGameplayTaskCall::ExpandNode(class FKismetCompilerContext& Co
 
 	bIsErrorFree &= Schema->TryCreateConnection(SpawnedActorReturnPin, InSpawnedActorPin);
 
-	LastThenPin = CallPostSpawnnProxyObjectNode->FindPinChecked(Schema->PN_Then);
+	LastThenPin = CallPostSpawnnProxyObjectNode->FindPinChecked(UEdGraphSchema_K2::PN_Then);
 
 
 	// Move the connections from the original node then pin to the last internal then pin
-	bIsErrorFree &= CompilerContext.MovePinLinksToIntermediate(*FindPinChecked(Schema->PN_Then), *LastThenPin).CanSafeConnect();
+	bIsErrorFree &= CompilerContext.MovePinLinksToIntermediate(*FindPinChecked(UEdGraphSchema_K2::PN_Then), *LastThenPin).CanSafeConnect();
 	bIsErrorFree &= CompilerContext.CopyPinLinksToIntermediate(*LastThenPin, *BranchElsePin).CanSafeConnect();
 	bIsErrorFree &= CompilerContext.CopyPinLinksToIntermediate(*LastThenPin, *ValidateProxyNode->GetElsePin()).CanSafeConnect();
 	
diff --git a/Engine/Source/Editor/GraphEditor/Private/DragConnection.cpp b/Engine/Source/Editor/GraphEditor/Private/DragConnection.cpp
--- a/Engine/Source/Editor/GraphEditor/Private/DragConnection.cpp
+++ b/Engine/Source/Editor/GraphEditor/Private/DragConnection.cpp
@@ -259,7 +259,7 @@ FReply FDragConnection::DroppedOnNode(FVector2D ScreenPosition, FVector2D GraphP
 					bHandledPinDropOnNode = true;
 
 					// Find which pin name to use and drop the pin on the node
-					FString PinName = SourcePin->PinFriendlyName.IsEmpty()? SourcePin->PinName : SourcePin->PinFriendlyName.ToString();
+					const FName PinName = SourcePin->PinFriendlyName.IsEmpty()? SourcePin->PinName : *SourcePin->PinFriendlyName.ToString();
 
 					const FScopedTransaction Transaction( NSLOCTEXT("UnrealEd", "AddInParam", "Add In Parameter" ) );
 
diff --git a/Engine/Source/Editor/GraphEditor/Private/GraphDiffControl.cpp b/Engine/Source/Editor/GraphEditor/Private/GraphDiffControl.cpp
--- a/Engine/Source/Editor/GraphEditor/Private/GraphDiffControl.cpp
+++ b/Engine/Source/Editor/GraphEditor/Private/GraphDiffControl.cpp
@@ -103,7 +103,7 @@ static void DiffR_NodeMoved(const FGraphDiffControl::FNodeDiffContext& DiffConte
 }
 
 /** Diff result when a pin type was changed */
-static void DiffR_PinTypeChanged(FDiffResults& Results, class UEdGraphPin* Pin2, class UEdGraphPin* Pin1)
+static void DiffR_PinTypeChanged(FDiffResults& Results, UEdGraphPin* Pin2, UEdGraphPin* Pin1)
 {
 	FEdGraphPinType Type1 = Pin1->PinType;
 	FEdGraphPinType Type2 = Pin2->PinType;
@@ -120,9 +120,9 @@ static void DiffR_PinTypeChanged(FDiffResults& Results, class UEdGraphPin* Pin2,
 		// Only bother setting up the display data if we're storing the result
 		if(Results.CanStoreResults())
 		{
-			Diff.ToolTip = FText::Format(LOCTEXT("DIF_PinCategoryToolTipFmt", "Pin '{0}' Category was '{1}', but is now '{2}"), FText::FromString(Pin2->PinName), FText::FromString(Pin1->PinType.PinCategory), FText::FromString(Pin2->PinType.PinCategory));
+			Diff.ToolTip = FText::Format(LOCTEXT("DIF_PinCategoryToolTipFmt", "Pin '{0}' Category was '{1}', but is now '{2}"), FText::FromName(Pin2->PinName), FText::FromName(Pin1->PinType.PinCategory), FText::FromName(Pin2->PinType.PinCategory));
 			Diff.DisplayColor = FLinearColor(0.15f,0.53f,0.15f);
-			Diff.DisplayString = FText::Format(LOCTEXT("DIF_PinCategoryFmt", "Pin Category '{0}' ['{1}' -> '{2}']"), FText::FromString(Pin2->PinName), FText::FromString(Pin1->PinType.PinCategory), FText::FromString(Pin2->PinType.PinCategory));
+			Diff.DisplayString = FText::Format(LOCTEXT("DIF_PinCategoryFmt", "Pin Category '{0}' ['{1}' -> '{2}']"), FText::FromName(Pin2->PinName), FText::FromName(Pin1->PinType.PinCategory), FText::FromName(Pin2->PinType.PinCategory));
 		}
 	}
 	else if(Type1.PinSubCategory != Type2.PinSubCategory)
@@ -132,9 +132,9 @@ static void DiffR_PinTypeChanged(FDiffResults& Results, class UEdGraphPin* Pin2,
 		// Only bother setting up the display data if we're storing the result
 		if(Results.CanStoreResults())
 		{
-			Diff.ToolTip = FText::Format(LOCTEXT("DIF_PinSubCategoryToolTipFmt", "Pin '{0}' SubCategory was '{1}', but is now '{2}"), FText::FromString(Pin2->PinName), FText::FromString(Pin1->PinType.PinSubCategory), FText::FromString(Pin2->PinType.PinSubCategory));
+			Diff.ToolTip = FText::Format(LOCTEXT("DIF_PinSubCategoryToolTipFmt", "Pin '{0}' SubCategory was '{1}', but is now '{2}"), FText::FromName(Pin2->PinName), FText::FromName(Pin1->PinType.PinSubCategory), FText::FromName(Pin2->PinType.PinSubCategory));
 			Diff.DisplayColor = FLinearColor(0.45f,0.53f,0.65f);
-			Diff.DisplayString = FText::Format(LOCTEXT("DIF_PinSubCategoryFmt", "Pin SubCategory '{0}'  ['{1}' -> '{2}']"), FText::FromString(Pin2->PinName), FText::FromString(Pin1->PinType.PinSubCategory), FText::FromString(Pin2->PinType.PinSubCategory));
+			Diff.DisplayString = FText::Format(LOCTEXT("DIF_PinSubCategoryFmt", "Pin SubCategory '{0}'  ['{1}' -> '{2}']"), FText::FromName(Pin2->PinName), FText::FromName(Pin1->PinType.PinSubCategory), FText::FromName(Pin2->PinType.PinSubCategory));
 		}
 	}
 	else if(T1Obj != T2Obj && (T1Obj && T2Obj ) &&
@@ -145,12 +145,12 @@ static void DiffR_PinTypeChanged(FDiffResults& Results, class UEdGraphPin* Pin2,
 		// Only bother setting up the display data if we're storing the result
 		if(Results.CanStoreResults())
 		{
-			FString Obj1 = T1Obj->GetFName().ToString();
-			FString Obj2 = T2Obj->GetFName().ToString();
+			const FName Obj1 = T1Obj->GetFName();
+			const FName Obj2 = T2Obj->GetFName();
 
-			Diff.ToolTip = FText::Format(LOCTEXT("DIF_PinSubCategorObjToolTipFmt", "Pin '{0}' was SubCategoryObject '{1}', but is now '{2}"), FText::FromString(Pin2->PinName), FText::FromString(Obj1), FText::FromString(Obj2));
+			Diff.ToolTip = FText::Format(LOCTEXT("DIF_PinSubCategorObjToolTipFmt", "Pin '{0}' was SubCategoryObject '{1}', but is now '{2}"), FText::FromName(Pin2->PinName), FText::FromName(Obj1), FText::FromName(Obj2));
 			Diff.DisplayColor = FLinearColor(0.45f,0.13f,0.25f);
-			Diff.DisplayString = FText::Format(LOCTEXT("DIF_PinSubCategoryObjFmt", "Pin SubCategoryObject '{0}' ['{1}' -> '{2}']"), FText::FromString(Pin2->PinName), FText::FromString(Obj1), FText::FromString(Obj2));
+			Diff.DisplayString = FText::Format(LOCTEXT("DIF_PinSubCategoryObjFmt", "Pin SubCategoryObject '{0}' ['{1}' -> '{2}']"), FText::FromName(Pin2->PinName), FText::FromName(Obj1), FText::FromName(Obj2));
 		}
 	}
 	else if(Type1.ContainerType != Type2.ContainerType)
@@ -164,8 +164,8 @@ static void DiffR_PinTypeChanged(FDiffResults& Results, class UEdGraphPin* Pin2,
 			FText IsArray1 = Pin1->PinType.IsArray() ? LOCTEXT("true", "true") : LOCTEXT("false", "false");
 			FText IsArray2 = Pin2->PinType.IsArray() ? LOCTEXT("true", "true") : LOCTEXT("false", "false");
 
-			Diff.ToolTip = FText::Format(LOCTEXT("DIF_PinIsArrayToolTipFmt", "PinType IsArray for '{0}' modified. Was '{1}', but is now '{2}"), FText::FromString(Pin2->PinName), IsArray1, IsArray2);
-			Diff.DisplayString = FText::Format(LOCTEXT("DIF_PinIsArrayFmt", "Pin IsArray '{0}' ['{1}' -> '{2}']"), FText::FromString(Pin2->PinName), IsArray1, IsArray2);
+			Diff.ToolTip = FText::Format(LOCTEXT("DIF_PinIsArrayToolTipFmt", "PinType IsArray for '{0}' modified. Was '{1}', but is now '{2}"), FText::FromName(Pin2->PinName), IsArray1, IsArray2);
+			Diff.DisplayString = FText::Format(LOCTEXT("DIF_PinIsArrayFmt", "Pin IsArray '{0}' ['{1}' -> '{2}']"), FText::FromName(Pin2->PinName), IsArray1, IsArray2);
 			Diff.DisplayColor = FLinearColor(0.45f,0.33f,0.35f);
 		}
 	}
@@ -179,9 +179,9 @@ static void DiffR_PinTypeChanged(FDiffResults& Results, class UEdGraphPin* Pin2,
 			FText IsRef1 = Pin1->PinType.bIsReference ? LOCTEXT("true", "true") : LOCTEXT("false", "false");
 			FText IsRef2 = Pin2->PinType.bIsReference ? LOCTEXT("true", "true") : LOCTEXT("false", "false");
 
-			Diff.ToolTip = FText::Format(LOCTEXT("DIF_PinIsRefToolTipFmt", "PinType IsReference for '{0}' modified. Was '{1}', but is now '{2}"), FText::FromString(Pin2->PinName), IsRef1, IsRef2);
+			Diff.ToolTip = FText::Format(LOCTEXT("DIF_PinIsRefToolTipFmt", "PinType IsReference for '{0}' modified. Was '{1}', but is now '{2}"), FText::FromName(Pin2->PinName), IsRef1, IsRef2);
 			Diff.DisplayColor = FLinearColor(0.25f,0.43f,0.35f);
-			Diff.DisplayString = FText::Format(LOCTEXT("DIF_PinIsRefFmt", "Pin IsReference '{0}' ['{1}' -> '{2}']"), FText::FromString(Pin2->PinName), IsRef1, IsRef2);
+			Diff.DisplayString = FText::Format(LOCTEXT("DIF_PinIsRefFmt", "Pin IsReference '{0}' ['{1}' -> '{2}']"), FText::FromName(Pin2->PinName), IsRef1, IsRef2);
 		}
 	}
 
@@ -191,7 +191,7 @@ static void DiffR_PinTypeChanged(FDiffResults& Results, class UEdGraphPin* Pin2,
 }
 
 /** Diff result when the # of links to a pin was changed */
-static void DiffR_PinLinkCountChanged(FDiffResults& Results, class UEdGraphPin* Pin2, class UEdGraphPin* Pin1)
+static void DiffR_PinLinkCountChanged(FDiffResults& Results, UEdGraphPin* Pin2, UEdGraphPin* Pin1)
 {
 	FDiffSingleResult Diff;
 	Diff.Diff = Pin2->LinkedTo.Num() > Pin1->LinkedTo.Num()  ?  EDiffType::PIN_LINKEDTO_NUM_INC : EDiffType::PIN_LINKEDTO_NUM_DEC;
@@ -203,15 +203,15 @@ static void DiffR_PinLinkCountChanged(FDiffResults& Results, class UEdGraphPin*
 	{
 		if(Diff.Diff == EDiffType::PIN_LINKEDTO_NUM_INC)
 		{
-			Diff.ToolTip = FText::Format(LOCTEXT("DIF_PinLinkCountIncToolTipFmt", "Pin '{0}' has more links (was {1} now {2})"), FText::FromString(Pin1->PinName), FText::AsNumber(Pin1->LinkedTo.Num()), FText::AsNumber(Pin2->LinkedTo.Num()));
+			Diff.ToolTip = FText::Format(LOCTEXT("DIF_PinLinkCountIncToolTipFmt", "Pin '{0}' has more links (was {1} now {2})"), FText::FromName(Pin1->PinName), FText::AsNumber(Pin1->LinkedTo.Num()), FText::AsNumber(Pin2->LinkedTo.Num()));
 			Diff.DisplayColor = FLinearColor(0.5f,0.3f,0.85f);
-			Diff.DisplayString = FText::Format(LOCTEXT("DIF_PinLinkCountIncFmt", "Added Link to '{0}'"), FText::FromString(Pin1->PinName));
+			Diff.DisplayString = FText::Format(LOCTEXT("DIF_PinLinkCountIncFmt", "Added Link to '{0}'"), FText::FromName(Pin1->PinName));
 		}
 		else
 		{
-			Diff.ToolTip = FText::Format(LOCTEXT("DIF_PinLinkCountDecToolTipFmt", "Pin '{0}' has fewer links (was {1} now {2})"), FText::FromString(Pin1->PinName), FText::AsNumber(Pin1->LinkedTo.Num()), FText::AsNumber(Pin2->LinkedTo.Num()));
+			Diff.ToolTip = FText::Format(LOCTEXT("DIF_PinLinkCountDecToolTipFmt", "Pin '{0}' has fewer links (was {1} now {2})"), FText::FromName(Pin1->PinName), FText::AsNumber(Pin1->LinkedTo.Num()), FText::AsNumber(Pin2->LinkedTo.Num()));
 			Diff.DisplayColor = FLinearColor(0.5f,0.3f,0.85f);
-			Diff.DisplayString = FText::Format(LOCTEXT("DIF_PinLinkCountDecFmt", "Removed Link to '{0}'"), FText::FromString(Pin1->PinName));
+			Diff.DisplayString = FText::Format(LOCTEXT("DIF_PinLinkCountDecFmt", "Removed Link to '{0}'"), FText::FromName(Pin1->PinName));
 		}
 	}
 
@@ -219,7 +219,7 @@ static void DiffR_PinLinkCountChanged(FDiffResults& Results, class UEdGraphPin*
 }
 
 /** Diff result when a pin to relinked to a different node */
-static void DiffR_LinkedToNode(FDiffResults& Results, class UEdGraphPin* Pin1, class UEdGraphPin* Pin2,  class UEdGraphNode* Node1, class UEdGraphNode* Node2)
+static void DiffR_LinkedToNode(FDiffResults& Results, UEdGraphPin* Pin1, UEdGraphPin* Pin2,  UEdGraphNode* Node1, UEdGraphNode* Node2)
 {
 	FDiffSingleResult Diff;
 	Diff.Diff = EDiffType::PIN_LINKEDTO_NODE;
@@ -235,7 +235,7 @@ static void DiffR_LinkedToNode(FDiffResults& Results, class UEdGraphPin* Pin1, c
 		FText Node2Name = Node2->GetNodeTitle(ENodeTitleType::ListView);
 
 		FFormatNamedArguments Args;
-		Args.Add(TEXT("PinNameForNode1"), FText::FromString(Pin1->PinName));
+		Args.Add(TEXT("PinNameForNode1"), FText::FromName(Pin1->PinName));
 		Args.Add(TEXT("NodeName1"), Node1Name);
 		Args.Add(TEXT("NodeName2"), Node2Name);
 		Diff.ToolTip = FText::Format(LOCTEXT("DIF_PinLinkMovedToolTip", "Pin '{PinNameForNode1}' was linked to Node '{NodeName1}', but is now linked to Node '{NodeName2}'"), Args);
@@ -248,7 +248,7 @@ static void DiffR_LinkedToNode(FDiffResults& Results, class UEdGraphPin* Pin1, c
 }
 
 /** Diff result when a pin default value was changed, and is in use*/
-static void DiffR_PinDefaultValueChanged(FDiffResults& Results, class UEdGraphPin* Pin2, class UEdGraphPin* Pin1)
+static void DiffR_PinDefaultValueChanged(FDiffResults& Results, UEdGraphPin* Pin2, UEdGraphPin* Pin1)
 {
 	FDiffSingleResult Diff;
 	Diff.Diff = EDiffType::PIN_DEFAULT_VALUE;
@@ -259,7 +259,7 @@ static void DiffR_PinDefaultValueChanged(FDiffResults& Results, class UEdGraphPi
 	if(Results.CanStoreResults())
 	{
 		FFormatNamedArguments Args;
-		Args.Add(TEXT("PinNameForValue1"), FText::FromString(Pin2->PinName));
+		Args.Add(TEXT("PinNameForValue1"), FText::FromName(Pin2->PinName));
 		Args.Add(TEXT("PinValue1"), FText::FromString(Pin1->GetDefaultAsString()));
 		Args.Add(TEXT("PinValue2"), FText::FromString(Pin2->GetDefaultAsString()));
 		Diff.ToolTip = FText::Format(LOCTEXT("DIF_PinDefaultValueToolTip", "Pin '{PinNameForValue1}' Default Value was '{PinValue1}', but is now '{PinValue2}"), Args);
@@ -279,7 +279,7 @@ static void DiffR_PinDefaultValueChanged(FDiffResults& Results, class UEdGraphPi
  * @param	Pins1	Array of relevant pins on Node 1.
  * @param	Pins2	Array of relevant pins on Node 2.
  */
-static void DiffR_NodePinCount(FDiffResults& Results, class UEdGraphNode* Node, class UEdGraphNode* Node2, const TArray<class UEdGraphPin*> &Pins1, const TArray<class UEdGraphPin*> &Pins2)
+static void DiffR_NodePinCount(FDiffResults& Results, UEdGraphNode* Node, UEdGraphNode* Node2, const TArray<UEdGraphPin*> &Pins1, const TArray<UEdGraphPin*> &Pins2)
 {
 	FText NodeName = Node->GetNodeTitle(ENodeTitleType::ListView);
 	int32 OriginalCount = Pins2.Num();
@@ -300,12 +300,12 @@ static void DiffR_NodePinCount(FDiffResults& Results, class UEdGraphNode* Node,
 
 		struct FMatchName
 		{
-			FMatchName(const FString& InPinName)
+			FMatchName(const FName InPinName)
 				: PinName(InPinName)
 			{
 			}
 
-			const FString& PinName;
+			const FName PinName;
 
 			bool operator()(const UEdGraphPin* Entry )
 			{
@@ -315,7 +315,7 @@ static void DiffR_NodePinCount(FDiffResults& Results, class UEdGraphNode* Node,
 
 		FText ListOfPins;
 		TArray< FText > RemovedPins;
-		for (auto OldPin : Pins1)
+		for (UEdGraphPin* OldPin : Pins1)
 		{
 			const UEdGraphPin* const* FoundPin = Pins2.FindByPredicate(FMatchName(OldPin->PinName));
 			if (FoundPin == nullptr)
@@ -325,7 +325,7 @@ static void DiffR_NodePinCount(FDiffResults& Results, class UEdGraphNode* Node,
 		}
 
 		TArray< FText > AddedPins;
-		for (auto OldPin : Pins2)
+		for (UEdGraphPin* OldPin : Pins2)
 		{
 			const UEdGraphPin* const* FoundPin = Pins1.FindByPredicate(FMatchName(OldPin->PinName));
 			if (FoundPin == nullptr)
@@ -366,7 +366,7 @@ static void DiffR_NodePinCount(FDiffResults& Results, class UEdGraphNode* Node,
 		if (AddedPins.Num() > 0)
 		{
 			Builder.AppendLine(LOCTEXT("DIF_PinsAddedList", "Pins Added:"));
-			for (const auto& Added : AddedPins)
+			for (const FText& Added : AddedPins)
 			{
 				Builder.AppendLine(Added);
 			}
@@ -375,7 +375,7 @@ static void DiffR_NodePinCount(FDiffResults& Results, class UEdGraphNode* Node,
 		if (RemovedPins.Num() > 0)
 		{
 			Builder.AppendLine(LOCTEXT("DIF_PinsRemovedList", "Pins Removed:"));
-			for (const auto& Removed : RemovedPins)
+			for (const FText& Removed : RemovedPins)
 			{
 				Builder.AppendLine(Removed);
 			}
@@ -392,7 +392,7 @@ static void DiffR_NodePinCount(FDiffResults& Results, class UEdGraphNode* Node,
  * @param	InPins			Pins
  * @param	OutRelevanPins	Output of Relevant pins
  */
-static void BuildArrayOfRelevantPins(const TArray<class UEdGraphPin*>& InPins, TArray< UEdGraphPin* >& OutRelevantPins)
+static void BuildArrayOfRelevantPins(const TArray<UEdGraphPin*>& InPins, TArray< UEdGraphPin* >& OutRelevantPins)
 {
 	for(int32 i = 0;i<InPins.Num();++i)
 	{
@@ -425,13 +425,12 @@ static bool IsPinTypeDifferent(const FEdGraphPinType& T1, const FEdGraphPinType&
 }
 
 /** Find Pin in array that matches pin */
-static class UEdGraphPin* FindOtherLink(TArray<class UEdGraphPin*>& Links2, int32 OriginalIndex, class UEdGraphPin* PinToFind)
+static UEdGraphPin* FindOtherLink(TArray<UEdGraphPin*>& Links2, int32 OriginalIndex, UEdGraphPin* PinToFind)
 {
 	//sometimes the order of the pins is different between revisions, although the pins themselves are unchanged, so we have to look at all of them
 	UEdGraphNode* Node1 = PinToFind->GetOwningNode();
-	for(auto It(Links2.CreateIterator());It;++It)
+	for (UEdGraphPin* Other : Links2)
 	{
-		UEdGraphPin* Other = *It;
 		UEdGraphNode* Node2 = Other->GetOwningNode();
 		if(FGraphDiffControl::IsNodeMatch(Node1, Node2))
 		{
@@ -442,7 +441,7 @@ static class UEdGraphPin* FindOtherLink(TArray<class UEdGraphPin*>& Links2, int3
 }
 
 /** Determine if the LinkedTo pins are the same */
-static bool LinkedToDifferent(class UEdGraphPin* OriginalPin1, class UEdGraphPin* OriginalPin2, const TArray<class UEdGraphPin*>& Links1, TArray<class UEdGraphPin*>& Links2, FDiffResults& Results)
+static bool LinkedToDifferent(UEdGraphPin* OriginalPin1, UEdGraphPin* OriginalPin2, const TArray<UEdGraphPin*>& Links1, TArray<UEdGraphPin*>& Links2, FDiffResults& Results)
 {
 	const int32 Size = Links1.Num();
 	bool bHasResult = false;
@@ -471,7 +470,7 @@ static bool LinkedToDifferent(class UEdGraphPin* OriginalPin1, class UEdGraphPin
  * 
  * returns true if any pins are different and populates the Results array
  */
-static bool ArePinsDifferent(const TArray<class UEdGraphPin*>& Pins1, TArray<class UEdGraphPin*>& Pins2, FDiffResults& Results)
+static bool ArePinsDifferent(const TArray<UEdGraphPin*>& Pins1, TArray<UEdGraphPin*>& Pins2, FDiffResults& Results)
 {
 	const int32 Size = Pins1.Num();
 	bool bHasResult = false;
@@ -598,22 +597,16 @@ FGraphDiffControl::FNodeMatch FGraphDiffControl::FindNodeMatch(class UEdGraph* G
 	FNodeMatch Match;
 	Match.NewNode = Node;
 
-	if (Graph != NULL)
+	if (Graph)
 	{
 		// attempt to find a node matching 'Node'
-		for (auto NodeIt(Graph->Nodes.CreateIterator()); NodeIt; ++NodeIt) 
+		for (UEdGraphNode* GraphNode : Graph->Nodes) 
 		{
-			UEdGraphNode* GraphNode = *NodeIt;
-			if (GraphNode == NULL)
-			{
-				continue;
-			}
-
-			if (IsNodeMatch(Node, GraphNode, &PriorMatches))
+			if (GraphNode && IsNodeMatch(Node, GraphNode, &PriorMatches))
 			{
 				Match.OldNode = GraphNode;
 				break;
-			}		
+			}
 		}
 	}	
 
@@ -642,12 +635,11 @@ bool FGraphDiffControl::IsNodeMatch(class UEdGraphNode* Node1, class UEdGraphNod
 		return (Node1->GetFName() == Node2->GetFName());
 	}
 
-	if (Exclusions != NULL)
+	if (Exclusions)
 	{
 		// have to see if this node has already been matched with another
-		for (auto MatchIt(Exclusions->CreateConstIterator()); MatchIt; ++MatchIt)
+		for (const FGraphDiffControl::FNodeMatch& PriorMatch : *Exclusions)
 		{
-			FGraphDiffControl::FNodeMatch const& PriorMatch = *MatchIt;
 			if (!PriorMatch.IsValid())
 			{
 				continue;
@@ -676,7 +668,7 @@ bool FGraphDiffControl::DiffGraphs(UEdGraph* const LhsGraph, UEdGraph* const Rhs
 {
 	bool bFoundDifferences = false;
 
-	if ((LhsGraph != NULL) && (RhsGraph != NULL))
+	if (LhsGraph && RhsGraph)
 	{
 		TArray<FGraphDiffControl::FNodeMatch> NodeMatches;
 		TSet<UEdGraphNode const*> MatchedRhsNodes;
@@ -685,24 +677,21 @@ bool FGraphDiffControl::DiffGraphs(UEdGraph* const LhsGraph, UEdGraph* const Rhs
 		AdditiveDiffContext.NodeTypeDisplayName = LOCTEXT("NodeDiffDisplayName", "Node");
 
 		// march through the all the nodes in the rhs graph and look for matches 
-		for (auto NodeIt(RhsGraph->Nodes.CreateConstIterator()); NodeIt; ++NodeIt)
+		for (UEdGraphNode* const RhsNode : RhsGraph->Nodes)
 		{
-			UEdGraphNode* const RhsNode = *NodeIt;
-			if (RhsNode == NULL)
+			if (RhsNode)
 			{
-				continue;
+				FGraphDiffControl::FNodeMatch NodeMatch = FGraphDiffControl::FindNodeMatch(LhsGraph, RhsNode, NodeMatches);
+				// if we found a corresponding node in the lhs graph, track it (so we
+				// can prevent future matches with the same nodes)
+				if (NodeMatch.IsValid())
+				{
+					NodeMatches.Add(NodeMatch);
+					MatchedRhsNodes.Add(NodeMatch.OldNode);
+				}
+
+				bFoundDifferences |= NodeMatch.Diff(AdditiveDiffContext, &DiffsOut);
 			}
-
-			FGraphDiffControl::FNodeMatch NodeMatch = FGraphDiffControl::FindNodeMatch(LhsGraph, RhsNode, NodeMatches);
-			// if we found a corresponding node in the lhs graph, track it (so we
-			// can prevent future matches with the same nodes)
-			if (NodeMatch.IsValid())
-			{
-				NodeMatches.Add(NodeMatch);
-				MatchedRhsNodes.Add(NodeMatch.OldNode);
-			}
-
-			bFoundDifferences |= NodeMatch.Diff(AdditiveDiffContext, &DiffsOut);
 		}
 
 		FGraphDiffControl::FNodeDiffContext SubtractiveDiffContext = AdditiveDiffContext;
@@ -710,11 +699,10 @@ bool FGraphDiffControl::DiffGraphs(UEdGraph* const LhsGraph, UEdGraph* const Rhs
 		SubtractiveDiffContext.DiffFlags = EDiffFlags::NodeExistance;
 
 		// go through the lhs nodes to catch ones that may have been missing from the rhs graph
-		for (auto NodeIt(LhsGraph->Nodes.CreateConstIterator()); NodeIt; ++NodeIt)
+		for (UEdGraphNode* const LhsNode : LhsGraph->Nodes)
 		{
-			UEdGraphNode* const LhsNode = *NodeIt;
 			// if this node has already been matched, move on
-			if ((LhsNode == NULL) || MatchedRhsNodes.Find(LhsNode))
+			if ((LhsNode == nullptr) || MatchedRhsNodes.Find(LhsNode))
 			{
 				continue;
 			}
@@ -725,8 +713,8 @@ bool FGraphDiffControl::DiffGraphs(UEdGraph* const LhsGraph, UEdGraph* const Rhs
 	}
 
 	// storing the graph name for all diff entries:
-	FName GraphName = LhsGraph ? LhsGraph->GetFName() : RhsGraph->GetFName();
-	for( auto& Entry : DiffsOut )
+	const FName GraphName = LhsGraph ? LhsGraph->GetFName() : RhsGraph->GetFName();
+	for( FDiffSingleResult& Entry : DiffsOut )
 	{
 		Entry.OwningGraph = GraphName;
 	}
diff --git a/Engine/Source/Editor/GraphEditor/Private/KismetNodes/SGraphNodeCallParameterCollectionFunction.cpp b/Engine/Source/Editor/GraphEditor/Private/KismetNodes/SGraphNodeCallParameterCollectionFunction.cpp
--- a/Engine/Source/Editor/GraphEditor/Private/KismetNodes/SGraphNodeCallParameterCollectionFunction.cpp
+++ b/Engine/Source/Editor/GraphEditor/Private/KismetNodes/SGraphNodeCallParameterCollectionFunction.cpp
@@ -17,7 +17,7 @@ TSharedPtr<SGraphPin> SGraphNodeCallParameterCollectionFunction::CreatePinWidget
 	// Create a special pin class for the ParameterName pin
 	if (CallFunctionNode
 		&& Pin->PinName == TEXT("ParameterName") 
-		&& Pin->PinType.PinCategory == GetDefault<UEdGraphSchema_K2>()->PC_Name)
+		&& Pin->PinType.PinCategory == UEdGraphSchema_K2::PC_Name)
 	{
 		TArray<FName> NameList;
 
diff --git a/Engine/Source/Editor/GraphEditor/Private/KismetNodes/SGraphNodeMakeStruct.cpp b/Engine/Source/Editor/GraphEditor/Private/KismetNodes/SGraphNodeMakeStruct.cpp
--- a/Engine/Source/Editor/GraphEditor/Private/KismetNodes/SGraphNodeMakeStruct.cpp
+++ b/Engine/Source/Editor/GraphEditor/Private/KismetNodes/SGraphNodeMakeStruct.cpp
@@ -381,15 +381,14 @@ TSharedPtr<SGraphPin> SGraphNodeMakeStruct::CreatePinWidget(UEdGraphPin* Pin) co
 	UK2Node_MakeStruct* SetFieldsNode = CastChecked<UK2Node_MakeStruct>(GetNodeObj());
 	TSharedPtr<SGraphPin> ResultPin = FNodeFactory::CreatePinWidget(Pin);
 
-	const UEdGraphSchema_K2* K2Schema = Cast<const UEdGraphSchema_K2>(GraphNode->GetSchema());
-	if (!Pin->PinType.bIsReference && Pin->PinType.PinCategory != K2Schema->PC_Exec && Pin->Direction != EGPD_Output)
+	if (!Pin->PinType.bIsReference && Pin->PinType.PinCategory != UEdGraphSchema_K2::PC_Exec && Pin->Direction != EGPD_Output)
 	{
 		FOptionalPinFromProperty ReturnedPropertyEntry;
 		for (FOptionalPinFromProperty& PropertyEntry : SetFieldsNode->ShowPinForProperties)
 		{
 			if (PropertyEntry.bHasOverridePin)
 			{
-				if (PropertyEntry.PropertyName.ToString() == Pin->PinName)
+				if (PropertyEntry.PropertyName == Pin->PinName)
 				{
 					TWeakPtr<SHorizontalBox> HorizontalPin = ResultPin->GetFullPinHorizontalRowWidget();
 					if (HorizontalPin.IsValid())
diff --git a/Engine/Source/Editor/GraphEditor/Private/KismetPins/SGraphPinInteger.cpp b/Engine/Source/Editor/GraphEditor/Private/KismetPins/SGraphPinInteger.cpp
--- a/Engine/Source/Editor/GraphEditor/Private/KismetPins/SGraphPinInteger.cpp
+++ b/Engine/Source/Editor/GraphEditor/Private/KismetPins/SGraphPinInteger.cpp
@@ -17,143 +17,140 @@ TSharedRef<SWidget>	SGraphPinInteger::GetDefaultValueWidget()
 {
 	check(GraphPinObj);
 
-	if (const UEdGraphSchema_K2* K2Schema = Cast<const UEdGraphSchema_K2>(GraphPinObj->GetSchema()))
-	{
-		FString PinSubCategory = GraphPinObj->PinType.PinSubCategory;
-		const UEnum* BitmaskEnum = Cast<UEnum>(GraphPinObj->PinType.PinSubCategoryObject.Get());
+	const FName PinSubCategory = GraphPinObj->PinType.PinSubCategory;
+	const UEnum* BitmaskEnum = Cast<UEnum>(GraphPinObj->PinType.PinSubCategoryObject.Get());
 
-		if (PinSubCategory == K2Schema->PSC_Bitmask)
+	if (PinSubCategory == UEdGraphSchema_K2::PSC_Bitmask)
+	{
+		struct BitmaskFlagInfo
 		{
-			struct BitmaskFlagInfo
-			{
-				int32 Value;
-				FText DisplayName;
-				FText ToolTipText;
-			};
+			int32 Value;
+			FText DisplayName;
+			FText ToolTipText;
+		};
 
-			const int32 BitmaskBitCount = sizeof(int32) << 3;
+		const int32 BitmaskBitCount = sizeof(int32) << 3;
 
-			TArray<BitmaskFlagInfo> BitmaskFlags;
-			BitmaskFlags.Reserve(BitmaskBitCount);
+		TArray<BitmaskFlagInfo> BitmaskFlags;
+		BitmaskFlags.Reserve(BitmaskBitCount);
 
-			if(BitmaskEnum)
+		if(BitmaskEnum)
+		{
+			const bool bUseEnumValuesAsMaskValues = BitmaskEnum->GetBoolMetaData(FBlueprintMetadata::MD_UseEnumValuesAsMaskValuesInEditor);
+			auto AddNewBitmaskFlagLambda = [BitmaskEnum, &BitmaskFlags](int32 InEnumIndex, int32 InFlagValue)
 			{
-				const bool bUseEnumValuesAsMaskValues = BitmaskEnum->GetBoolMetaData(FBlueprintMetadata::MD_UseEnumValuesAsMaskValuesInEditor);
-				auto AddNewBitmaskFlagLambda = [BitmaskEnum, &BitmaskFlags](int32 InEnumIndex, int32 InFlagValue)
-				{
-					BitmaskFlagInfo* BitmaskFlag = new(BitmaskFlags) BitmaskFlagInfo();
+				BitmaskFlagInfo* BitmaskFlag = new(BitmaskFlags) BitmaskFlagInfo();
 
-					BitmaskFlag->Value = InFlagValue;
-					BitmaskFlag->DisplayName = BitmaskEnum->GetDisplayNameTextByIndex(InEnumIndex);
-					BitmaskFlag->ToolTipText = BitmaskEnum->GetToolTipTextByIndex(InEnumIndex);
-					if (BitmaskFlag->ToolTipText.IsEmpty())
-					{
-						BitmaskFlag->ToolTipText = FText::Format(NSLOCTEXT("GraphEditor", "BitmaskDefaultFlagToolTipText", "Toggle {0} on/off"), BitmaskFlag->DisplayName);
-					}
-				};
+				BitmaskFlag->Value = InFlagValue;
+				BitmaskFlag->DisplayName = BitmaskEnum->GetDisplayNameTextByIndex(InEnumIndex);
+				BitmaskFlag->ToolTipText = BitmaskEnum->GetToolTipTextByIndex(InEnumIndex);
+				if (BitmaskFlag->ToolTipText.IsEmpty())
+				{
+					BitmaskFlag->ToolTipText = FText::Format(NSLOCTEXT("GraphEditor", "BitmaskDefaultFlagToolTipText", "Toggle {0} on/off"), BitmaskFlag->DisplayName);
+				}
+			};
 
-				// Note: This loop doesn't include (BitflagsEnum->NumEnums() - 1) in order to skip the implicit "MAX" value that gets added to the enum type at compile time.
-				for(int32 BitmaskEnumIndex = 0; BitmaskEnumIndex < BitmaskEnum->NumEnums() - 1; ++BitmaskEnumIndex)
+			// Note: This loop doesn't include (BitflagsEnum->NumEnums() - 1) in order to skip the implicit "MAX" value that gets added to the enum type at compile time.
+			for(int32 BitmaskEnumIndex = 0; BitmaskEnumIndex < BitmaskEnum->NumEnums() - 1; ++BitmaskEnumIndex)
+			{
+				const int64 EnumValue = BitmaskEnum->GetValueByIndex(BitmaskEnumIndex);
+				if (EnumValue >= 0)
 				{
-					const int64 EnumValue = BitmaskEnum->GetValueByIndex(BitmaskEnumIndex);
-					if (EnumValue >= 0)
+					if (bUseEnumValuesAsMaskValues)
 					{
-						if (bUseEnumValuesAsMaskValues)
-						{
-							if (EnumValue < MAX_int32 && FMath::IsPowerOfTwo(EnumValue))
-							{
-								AddNewBitmaskFlagLambda(BitmaskEnumIndex, static_cast<int32>(EnumValue));
-							}
-						}
-						else if (EnumValue < BitmaskBitCount)
+						if (EnumValue < MAX_int32 && FMath::IsPowerOfTwo(EnumValue))
 						{
-							AddNewBitmaskFlagLambda(BitmaskEnumIndex, 1 << static_cast<int32>(EnumValue));
+							AddNewBitmaskFlagLambda(BitmaskEnumIndex, static_cast<int32>(EnumValue));
 						}
 					}
+					else if (EnumValue < BitmaskBitCount)
+					{
+						AddNewBitmaskFlagLambda(BitmaskEnumIndex, 1 << static_cast<int32>(EnumValue));
+					}
 				}
 			}
-			else
+		}
+		else
+		{
+			for(int BitmaskFlagIndex = 0; BitmaskFlagIndex < BitmaskBitCount; ++BitmaskFlagIndex)
 			{
-				for(int BitmaskFlagIndex = 0; BitmaskFlagIndex < BitmaskBitCount; ++BitmaskFlagIndex)
-				{
-					BitmaskFlagInfo* BitmaskFlag = new(BitmaskFlags) BitmaskFlagInfo();
+				BitmaskFlagInfo* BitmaskFlag = new(BitmaskFlags) BitmaskFlagInfo();
 
-					BitmaskFlag->Value = static_cast<int32>(1 << BitmaskFlagIndex);
-					BitmaskFlag->DisplayName = FText::Format(NSLOCTEXT("GraphEditor", "BitmaskDefaultFlagDisplayName", "Flag {0}"), FText::AsNumber(BitmaskFlagIndex + 1));
-					BitmaskFlag->ToolTipText = FText::Format(NSLOCTEXT("GraphEditor", "BitmaskDefaultFlagToolTipText", "Toggle {0} on/off"), BitmaskFlag->DisplayName);
-				}
+				BitmaskFlag->Value = static_cast<int32>(1 << BitmaskFlagIndex);
+				BitmaskFlag->DisplayName = FText::Format(NSLOCTEXT("GraphEditor", "BitmaskDefaultFlagDisplayName", "Flag {0}"), FText::AsNumber(BitmaskFlagIndex + 1));
+				BitmaskFlag->ToolTipText = FText::Format(NSLOCTEXT("GraphEditor", "BitmaskDefaultFlagToolTipText", "Toggle {0} on/off"), BitmaskFlag->DisplayName);
 			}
+		}
 
-			const auto& GetComboButtonText = [this, BitmaskFlags]() -> FText
+		const auto& GetComboButtonText = [this, BitmaskFlags]() -> FText
+		{
+			int32 BitmaskValue = FCString::Atoi(*GraphPinObj->GetDefaultAsString());
+			if(BitmaskValue != 0)
 			{
-				int32 BitmaskValue = FCString::Atoi(*GraphPinObj->GetDefaultAsString());
-				if(BitmaskValue != 0)
+				if(BitmaskValue & (BitmaskValue - 1))
 				{
-					if(BitmaskValue & (BitmaskValue - 1))
-					{
-						return NSLOCTEXT("GraphEditor", "BitmaskButtonContentMultipleBitsSet", "(Multiple)");
-					}
-					else
+					return NSLOCTEXT("GraphEditor", "BitmaskButtonContentMultipleBitsSet", "(Multiple)");
+				}
+				else
+				{
+					for(int i = 0; i < BitmaskFlags.Num(); ++i)
 					{
-						for(int i = 0; i < BitmaskFlags.Num(); ++i)
+						if(BitmaskValue & BitmaskFlags[i].Value)
 						{
-							if(BitmaskValue & BitmaskFlags[i].Value)
-							{
-								return BitmaskFlags[i].DisplayName;
-							}
+							return BitmaskFlags[i].DisplayName;
 						}
 					}
 				}
+			}
 
-				return NSLOCTEXT("GraphEditor", "BitmaskButtonContentNoFlagsSet", "(No Flags)");
-			};
-
-			return SNew(SComboButton)
-				.ContentPadding(3)
-				.MenuPlacement(MenuPlacement_BelowAnchor)
-				.Visibility(this, &SGraphPin::GetDefaultValueVisibility)
-				.ButtonContent()
+			return NSLOCTEXT("GraphEditor", "BitmaskButtonContentNoFlagsSet", "(No Flags)");
+		};
+
+		return SNew(SComboButton)
+			.ContentPadding(3)
+			.MenuPlacement(MenuPlacement_BelowAnchor)
+			.Visibility(this, &SGraphPin::GetDefaultValueVisibility)
+			.ButtonContent()
+			[
+				// Wrap in configurable box to restrain height/width of menu
+				SNew(SBox)
+				.MinDesiredWidth(84.0f)
 				[
-					// Wrap in configurable box to restrain height/width of menu
-					SNew(SBox)
-					.MinDesiredWidth(84.0f)
-					[
-						SNew(STextBlock)
-						.Text_Lambda(GetComboButtonText)
-						.Font(FEditorStyle::GetFontStyle(TEXT("PropertyWindow.NormalFont")))
-					]
+					SNew(STextBlock)
+					.Text_Lambda(GetComboButtonText)
+					.Font(FEditorStyle::GetFontStyle(TEXT("PropertyWindow.NormalFont")))
 				]
-				.OnGetMenuContent_Lambda([this, BitmaskFlags]()
-				{
-					FMenuBuilder MenuBuilder(false, nullptr);
+			]
+			.OnGetMenuContent_Lambda([this, BitmaskFlags]()
+			{
+				FMenuBuilder MenuBuilder(false, nullptr);
 
-					for(int i = 0; i < BitmaskFlags.Num(); ++i)
-					{
-						MenuBuilder.AddMenuEntry(
-							BitmaskFlags[i].DisplayName,
-							BitmaskFlags[i].ToolTipText,
-							FSlateIcon(),
-							FUIAction
-							(
-								FExecuteAction::CreateLambda([=]()
-								{
-									const int32 CurValue = FCString::Atoi(*GraphPinObj->GetDefaultAsString());
-									GraphPinObj->GetSchema()->TrySetDefaultValue(*GraphPinObj, FString::FromInt(CurValue ^ BitmaskFlags[i].Value));
-								}),
-								FCanExecuteAction(),
-								FIsActionChecked::CreateLambda([=]() -> bool
-								{
-									const int32 CurValue = FCString::Atoi(*GraphPinObj->GetDefaultAsString());
-									return (CurValue & BitmaskFlags[i].Value) != 0;
-								})
-							),
-							NAME_None,
-							EUserInterfaceActionType::Check);
-					}
+				for(int i = 0; i < BitmaskFlags.Num(); ++i)
+				{
+					MenuBuilder.AddMenuEntry(
+						BitmaskFlags[i].DisplayName,
+						BitmaskFlags[i].ToolTipText,
+						FSlateIcon(),
+						FUIAction
+						(
+							FExecuteAction::CreateLambda([=]()
+							{
+								const int32 CurValue = FCString::Atoi(*GraphPinObj->GetDefaultAsString());
+								GraphPinObj->GetSchema()->TrySetDefaultValue(*GraphPinObj, FString::FromInt(CurValue ^ BitmaskFlags[i].Value));
+							}),
+							FCanExecuteAction(),
+							FIsActionChecked::CreateLambda([=]() -> bool
+							{
+								const int32 CurValue = FCString::Atoi(*GraphPinObj->GetDefaultAsString());
+								return (CurValue & BitmaskFlags[i].Value) != 0;
+							})
+						),
+						NAME_None,
+						EUserInterfaceActionType::Check);
+				}
 
-					return MenuBuilder.MakeWidget();
-				});
-		}
+				return MenuBuilder.MakeWidget();
+			});
 	}
 	
 	return SGraphPinNum::GetDefaultValueWidget();
diff --git a/Engine/Source/Editor/GraphEditor/Private/KismetPins/SGraphPinObject.cpp b/Engine/Source/Editor/GraphEditor/Private/KismetPins/SGraphPinObject.cpp
--- a/Engine/Source/Editor/GraphEditor/Private/KismetPins/SGraphPinObject.cpp
+++ b/Engine/Source/Editor/GraphEditor/Private/KismetPins/SGraphPinObject.cpp
@@ -250,7 +250,7 @@ FText SGraphPinObject::GetValue() const
 	{
 		if (GraphPinObj->GetSchema()->IsSelfPin(*GraphPinObj))
 		{
-			Value =  FText::FromString(GraphPinObj->PinName);
+			Value =  FText::FromName(GraphPinObj->PinName);
 		}
 		else
 		{
diff --git a/Engine/Source/Editor/GraphEditor/Private/NodeFactory.cpp b/Engine/Source/Editor/GraphEditor/Private/NodeFactory.cpp
--- a/Engine/Source/Editor/GraphEditor/Private/NodeFactory.cpp
+++ b/Engine/Source/Editor/GraphEditor/Private/NodeFactory.cpp
@@ -232,51 +232,51 @@ TSharedPtr<SGraphPin> FNodeFactory::CreatePinWidget(UEdGraphPin* InPin)
 
 	if (const UEdGraphSchema_K2* K2Schema = Cast<const UEdGraphSchema_K2>(InPin->GetSchema()))
 	{
-		if (InPin->PinType.PinCategory == K2Schema->PC_Boolean)
+		if (InPin->PinType.PinCategory == UEdGraphSchema_K2::PC_Boolean)
 		{
 			return SNew(SGraphPinBool, InPin);
 		}
-		else if (InPin->PinType.PinCategory == K2Schema->PC_Text)
+		else if (InPin->PinType.PinCategory == UEdGraphSchema_K2::PC_Text)
 		{
 			return SNew(SGraphPinText, InPin);
 		}
-		else if (InPin->PinType.PinCategory == K2Schema->PC_Exec)
+		else if (InPin->PinType.PinCategory == UEdGraphSchema_K2::PC_Exec)
 		{
 			return SNew(SGraphPinExec, InPin);
 		}
-		else if (InPin->PinType.PinCategory == K2Schema->PC_Object)
+		else if (InPin->PinType.PinCategory == UEdGraphSchema_K2::PC_Object)
 		{
 			return SNew(SGraphPinObject, InPin);
 		}
-		else if (InPin->PinType.PinCategory == K2Schema->PC_Interface)
+		else if (InPin->PinType.PinCategory == UEdGraphSchema_K2::PC_Interface)
 		{
 			return SNew(SGraphPinObject, InPin);
 		}
-		else if (InPin->PinType.PinCategory == K2Schema->PC_SoftObject)
+		else if (InPin->PinType.PinCategory == UEdGraphSchema_K2::PC_SoftObject)
 		{
 			return SNew(SGraphPinObject, InPin);
 		}
-		else if (InPin->PinType.PinCategory == K2Schema->PC_Class)
+		else if (InPin->PinType.PinCategory == UEdGraphSchema_K2::PC_Class)
 		{
 			return SNew(SGraphPinClass, InPin);
 		}
-		else if (InPin->PinType.PinCategory == K2Schema->PC_SoftClass)
+		else if (InPin->PinType.PinCategory == UEdGraphSchema_K2::PC_SoftClass)
 		{
 			return SNew(SGraphPinClass, InPin);
 		}
-		else if (InPin->PinType.PinCategory == K2Schema->PC_Int)
+		else if (InPin->PinType.PinCategory == UEdGraphSchema_K2::PC_Int)
 		{
 			return SNew(SGraphPinInteger, InPin);
 		}
-		else if (InPin->PinType.PinCategory == K2Schema->PC_Float)
+		else if (InPin->PinType.PinCategory == UEdGraphSchema_K2::PC_Float)
 		{
 			return SNew(SGraphPinNum, InPin);
 		}
-		else if (InPin->PinType.PinCategory == K2Schema->PC_String || InPin->PinType.PinCategory == K2Schema->PC_Name)
+		else if (InPin->PinType.PinCategory == UEdGraphSchema_K2::PC_String || InPin->PinType.PinCategory == UEdGraphSchema_K2::PC_Name)
 		{
 			return SNew(SGraphPinString, InPin);
 		}
-		else if (InPin->PinType.PinCategory == K2Schema->PC_Struct)
+		else if (InPin->PinType.PinCategory == UEdGraphSchema_K2::PC_Struct)
 		{
 			// If you update this logic you'll probably need to update UEdGraphSchema_K2::ShouldHidePinDefaultValue!
 			UScriptStruct* ColorStruct = TBaseStructure<FLinearColor>::Get();
@@ -305,7 +305,7 @@ TSharedPtr<SGraphPin> FNodeFactory::CreatePinWidget(UEdGraphPin* InPin)
 				return SNew(SGraphPinCollisionProfile, InPin);
 			}
 		}
-		else if (InPin->PinType.PinCategory == K2Schema->PC_Byte)
+		else if (InPin->PinType.PinCategory == UEdGraphSchema_K2::PC_Byte)
 		{
 			// Check for valid enum object reference
 			if ((InPin->PinType.PinSubCategoryObject != NULL) && (InPin->PinType.PinSubCategoryObject->IsA(UEnum::StaticClass())))
@@ -317,11 +317,11 @@ TSharedPtr<SGraphPin> FNodeFactory::CreatePinWidget(UEdGraphPin* InPin)
 				return SNew(SGraphPinInteger, InPin);
 			}
 		}
-		else if ((InPin->PinType.PinCategory == K2Schema->PC_Wildcard) && (InPin->PinType.PinSubCategory == K2Schema->PSC_Index))
+		else if ((InPin->PinType.PinCategory == UEdGraphSchema_K2::PC_Wildcard) && (InPin->PinType.PinSubCategory == UEdGraphSchema_K2::PSC_Index))
 		{
 			return SNew(SGraphPinIndex, InPin);
 		}
-		else if(InPin->PinType.PinCategory == K2Schema->PC_MCDelegate)
+		else if(InPin->PinType.PinCategory == UEdGraphSchema_K2::PC_MCDelegate)
 		{
 			return SNew(SGraphPinString, InPin);
 		}
@@ -329,7 +329,7 @@ TSharedPtr<SGraphPin> FNodeFactory::CreatePinWidget(UEdGraphPin* InPin)
 
 	if (const UMaterialGraphSchema* MaterialGraphSchema = Cast<const UMaterialGraphSchema>(InPin->GetSchema()))
 	{
-		if (InPin->PinType.PinCategory == MaterialGraphSchema->PC_MaterialInput)
+		if (InPin->PinType.PinCategory == UMaterialGraphSchema::PC_MaterialInput)
 		{
 			return SNew(SGraphPinMaterialInput, InPin);
 		}
@@ -345,42 +345,46 @@ TSharedPtr<SGraphPin> FNodeFactory::CreatePinWidget(UEdGraphPin* InPin)
 
 FConnectionDrawingPolicy* FNodeFactory::CreateConnectionPolicy(const UEdGraphSchema* Schema, int32 InBackLayerID, int32 InFrontLayerID, float ZoomFactor, const FSlateRect& InClippingRect, FSlateWindowElementList& InDrawElements, UEdGraph* InGraphObj)
 {
-    FConnectionDrawingPolicy* ConnectionDrawingPolicy;
+	FConnectionDrawingPolicy* ConnectionDrawingPolicy;
 
-    // First give the schema a chance to provide the connection drawing policy
-    ConnectionDrawingPolicy = Schema->CreateConnectionDrawingPolicy(InBackLayerID, InFrontLayerID, ZoomFactor, InClippingRect, InDrawElements, InGraphObj);
+	// First give the schema a chance to provide the connection drawing policy
+	ConnectionDrawingPolicy = Schema->CreateConnectionDrawingPolicy(InBackLayerID, InFrontLayerID, ZoomFactor, InClippingRect, InDrawElements, InGraphObj);
 
-    // First give a shot to the registered connection factories
-    if (!ConnectionDrawingPolicy)
-    {
-        for (auto FactoryIt = FEdGraphUtilities::VisualPinConnectionFactories.CreateIterator(); FactoryIt; ++FactoryIt)
-        {
-            TSharedPtr<FGraphPanelPinConnectionFactory> FactoryPtr = *FactoryIt;
-            if (FactoryPtr.IsValid())
-            {
-                ConnectionDrawingPolicy = FactoryPtr->CreateConnectionPolicy(Schema, InBackLayerID, InFrontLayerID, ZoomFactor, InClippingRect, InDrawElements, InGraphObj);
-            }
-        }
-    }
+	// First give a shot to the registered connection factories
+	if (!ConnectionDrawingPolicy)
+	{
+		for (TSharedPtr<FGraphPanelPinConnectionFactory> FactoryPtr : FEdGraphUtilities::VisualPinConnectionFactories)
+		{
+			if (FactoryPtr.IsValid())
+			{
+				ConnectionDrawingPolicy = FactoryPtr->CreateConnectionPolicy(Schema, InBackLayerID, InFrontLayerID, ZoomFactor, InClippingRect, InDrawElements, InGraphObj);
+
+				if (ConnectionDrawingPolicy)
+				{
+					break;
+				}
+			}
+		}
+	}
 
-    // If neither the schema nor the factory provides a policy, try the hardcoded ones
-    //@TODO: Fold all of this code into registered factories for the various schemas!
-    if (!ConnectionDrawingPolicy)
-    {
+	// If neither the schema nor the factory provides a policy, try the hardcoded ones
+	//@TODO: Fold all of this code into registered factories for the various schemas!
+	if (!ConnectionDrawingPolicy)
+	{
 		if (Schema->IsA(UEdGraphSchema_K2::StaticClass()))
-        {
-            ConnectionDrawingPolicy = new FKismetConnectionDrawingPolicy(InBackLayerID, InFrontLayerID, ZoomFactor, InClippingRect, InDrawElements, InGraphObj);
-        }
-        else if (Schema->IsA(UMaterialGraphSchema::StaticClass()))
-        {
-            ConnectionDrawingPolicy = new FMaterialGraphConnectionDrawingPolicy(InBackLayerID, InFrontLayerID, ZoomFactor, InClippingRect, InDrawElements, InGraphObj);
-        }
-        else
-        {
-            ConnectionDrawingPolicy = new FConnectionDrawingPolicy(InBackLayerID, InFrontLayerID, ZoomFactor, InClippingRect, InDrawElements);
-        }
-    }
+		{
+			ConnectionDrawingPolicy = new FKismetConnectionDrawingPolicy(InBackLayerID, InFrontLayerID, ZoomFactor, InClippingRect, InDrawElements, InGraphObj);
+		}
+		else if (Schema->IsA(UMaterialGraphSchema::StaticClass()))
+		{
+			ConnectionDrawingPolicy = new FMaterialGraphConnectionDrawingPolicy(InBackLayerID, InFrontLayerID, ZoomFactor, InClippingRect, InDrawElements, InGraphObj);
+		}
+		else
+		{
+			ConnectionDrawingPolicy = new FConnectionDrawingPolicy(InBackLayerID, InFrontLayerID, ZoomFactor, InClippingRect, InDrawElements);
+		}
+	}
 
-    // If we never picked a custom policy, use the uncustomized standard policy
-    return ConnectionDrawingPolicy;
+	// If we never picked a custom policy, use the uncustomized standard policy
+	return ConnectionDrawingPolicy;
 }
\ No newline at end of file
diff --git a/Engine/Source/Editor/GraphEditor/Private/SGraphNode.cpp b/Engine/Source/Editor/GraphEditor/Private/SGraphNode.cpp
--- a/Engine/Source/Editor/GraphEditor/Private/SGraphNode.cpp
+++ b/Engine/Source/Editor/GraphEditor/Private/SGraphNode.cpp
@@ -1062,13 +1062,11 @@ void SGraphNode::CreateAdvancedViewArrow(TSharedPtr<SVerticalBox> MainBox)
 
 bool SGraphNode::ShouldPinBeHidden(const UEdGraphPin* InPin) const
 {
-	const UEdGraphSchema_K2* K2Schema = Cast<const UEdGraphSchema_K2>(GraphNode->GetSchema());
-
 	bool bHideNoConnectionPins = false;
 	bool bHideNoConnectionNoDefaultPins = false;
 
 	// Not allowed to hide exec pins 
-	const bool bCanHidePin = (K2Schema && (InPin->PinType.PinCategory != K2Schema->PC_Exec));
+	const bool bCanHidePin = (InPin->PinType.PinCategory != UEdGraphSchema_K2::PC_Exec);
 
 	if (OwnerGraphPanelPtr.IsValid() && bCanHidePin)
 	{
@@ -1078,7 +1076,7 @@ bool SGraphNode::ShouldPinBeHidden(const UEdGraphPin* InPin) const
 
 	const bool bIsOutputPin = InPin->Direction == EGPD_Output;
 	const bool bPinHasDefaultValue = !InPin->DefaultValue.IsEmpty() || (InPin->DefaultObject != NULL);
-	const bool bIsSelfTarget = K2Schema && (InPin->PinType.PinCategory == K2Schema->PC_Object) && (InPin->PinName == K2Schema->PN_Self);
+	const bool bIsSelfTarget = (InPin->PinType.PinCategory == UEdGraphSchema_K2::PC_Object) && (InPin->PinName == UEdGraphSchema_K2::PN_Self);
 	const bool bPinHasValidDefault = !bIsOutputPin && (bPinHasDefaultValue || bIsSelfTarget);
 	const bool bPinHasConections = InPin->LinkedTo.Num() > 0;
 
@@ -1115,7 +1113,7 @@ void SGraphNode::CreatePinWidgets()
 			, *GraphNode->GetNodeTitle(ENodeTitleType::ListView).ToString()
 			, *GraphNode->GetPathName()
 			, (CurPin->Direction == EEdGraphPinDirection::EGPD_Input) ? TEXT("input") : TEXT("output")
-			,  CurPin->PinFriendlyName.IsEmpty() ? *CurPin->PinName : *CurPin->PinFriendlyName.ToString()
+			,  CurPin->PinFriendlyName.IsEmpty() ? *CurPin->PinName.ToString() : *CurPin->PinFriendlyName.ToString()
 			,  CurPin->GetOuter() ? *CurPin->GetOuter()->GetClass()->GetName() : TEXT("UNKNOWN")
 			,  CurPin->GetOuter() ? *CurPin->GetOuter()->GetPathName() : TEXT("NULL")) )
 		{
diff --git a/Engine/Source/Editor/GraphEditor/Private/SGraphPin.cpp b/Engine/Source/Editor/GraphEditor/Private/SGraphPin.cpp
--- a/Engine/Source/Editor/GraphEditor/Private/SGraphPin.cpp
+++ b/Engine/Source/Editor/GraphEditor/Private/SGraphPin.cpp
@@ -361,7 +361,7 @@ FReply SGraphPin::OnPinMouseDown( const FGeometry& SenderGeometry, const FPointe
 				struct FLinkedToPinInfo
 				{
 					// Pin name string
-					FString PinName;
+					FName PinName;
 
 					// A weak reference to the node object that owns the pin
 					TWeakObjectPtr<UEdGraphNode> OwnerNodePtr;
@@ -371,7 +371,7 @@ FReply SGraphPin::OnPinMouseDown( const FGeometry& SenderGeometry, const FPointe
 				TArray<FLinkedToPinInfo> LinkedToPinInfoArray;
 				for (UEdGraphPin* Pin : GetPinObj()->LinkedTo)
 				{
-					if (auto PinWidget = PinToPinWidgetMap.Find(Pin))
+					if (TSharedRef<SGraphPin>* PinWidget = PinToPinWidgetMap.Find(Pin))
 					{
 						check((*PinWidget)->OwnerNodePtr.IsValid());
 
diff --git a/Engine/Source/Editor/Kismet/Private/BPVariableDragDropAction.cpp b/Engine/Source/Editor/Kismet/Private/BPVariableDragDropAction.cpp
--- a/Engine/Source/Editor/Kismet/Private/BPVariableDragDropAction.cpp
+++ b/Engine/Source/Editor/Kismet/Private/BPVariableDragDropAction.cpp
@@ -48,7 +48,7 @@ void FKismetVariableDragDropAction::GetLinksThatWillBreak(	UEdGraphNode* Node, U
 		{
 			FEdGraphPinType NewPinType;
 			Schema->ConvertPropertyToPinType(NewVariableProperty,NewPinType);
-			if(UEdGraphPin* Pin = VarNodeUnderCursor->FindPin(VarNodeUnderCursor->GetVarNameString()))
+			if(UEdGraphPin* Pin = VarNodeUnderCursor->FindPin(VarNodeUnderCursor->GetVarName()))
 			{
 				for(TArray<class UEdGraphPin*>::TIterator i(Pin->LinkedTo);i;i++)
 				{
@@ -128,7 +128,7 @@ void FKismetVariableDragDropAction::HoverTargetChanged()
 	else if (PinUnderCursor)
 	{
 		FFormatNamedArguments Args;
-		Args.Add(TEXT("PinUnderCursor"), FText::FromString(PinUnderCursor->PinName));
+		Args.Add(TEXT("PinUnderCursor"), FText::FromName(PinUnderCursor->PinName));
 		Args.Add(TEXT("VariableName"), FText::FromString(VariableString));
 
 		if (CanVariableBeDropped(VariableProperty, *PinUnderCursor->GetOwningNode()->GetGraph()))
@@ -151,7 +151,7 @@ void FKismetVariableDragDropAction::HoverTargetChanged()
 				Schema->ConvertPropertyToPinType(VariableProperty, VariablePinType);
 				const bool bTypeMatch = Schema->ArePinTypesCompatible(VariablePinType, PinUnderCursor->PinType) || bIsExecPin;
 
-				Args.Add(TEXT("PinUnderCursor"), FText::FromString(PinUnderCursor->PinName));
+				Args.Add(TEXT("PinUnderCursor"), FText::FromName(PinUnderCursor->PinName));
 
 				if (bTypeMatch && bCanWriteIfNeeded)
 				{
@@ -278,7 +278,7 @@ FReply FKismetVariableDragDropAction::DroppedOnNode(FVector2D ScreenPosition, FV
 
 		if(CanVariableBeDropped(VariableProperty, *TargetNode->GetGraph()))
 		{
-			const FString OldVarName = TargetNode->GetVarNameString();
+			const FName OldVarName = TargetNode->GetVarName();
 			const UEdGraphSchema_K2* Schema = Cast<const UEdGraphSchema_K2>(TargetNode->GetSchema());
 
 			TArray<class UEdGraphPin*> BadLinks;
@@ -307,7 +307,7 @@ FReply FKismetVariableDragDropAction::DroppedOnNode(FVector2D ScreenPosition, FV
 				FEdGraphPinType NewPinType;
 				Schema->ConvertPropertyToPinType(VariableProperty,NewPinType);
 
-				Pin->PinName = VariableName.ToString();
+				Pin->PinName = VariableName;
 				Pin->PinType = NewPinType;
 
 				//break bad links
diff --git a/Engine/Source/Editor/Kismet/Private/BlueprintCompilationManager.cpp b/Engine/Source/Editor/Kismet/Private/BlueprintCompilationManager.cpp
--- a/Engine/Source/Editor/Kismet/Private/BlueprintCompilationManager.cpp
+++ b/Engine/Source/Editor/Kismet/Private/BlueprintCompilationManager.cpp
@@ -24,12 +24,13 @@
 #include "ProfilingDebugging/ScopedTimers.h"
 #include "Serialization/ArchiveHasReferences.h"
 #include "Serialization/ArchiveReplaceOrClearExternalReferences.h"
+#include "Settings/EditorProjectSettings.h"
 #include "TickableEditorObject.h"
 #include "UObject/MetaData.h"
 #include "UObject/ReferenceChainSearch.h"
 #include "UObject/UObjectHash.h"
-#include "WidgetBlueprint.h"
 #include "Kismet2/KismetDebugUtilities.h"
+#include "BlueprintEditorModule.h"
 
 #define LOCTEXT_NAMESPACE "BlueprintCompilationManager"
 
@@ -91,6 +92,7 @@ struct FBlueprintCompilationManagerImpl : public FGCObject
 	void FlushReinstancingQueueImpl();
 	bool HasBlueprintsToCompile() const;
 	bool IsGeneratedClassLayoutReady() const;
+	void GetDefaultValue(const UClass* ForClass, const UProperty* Property, FString& OutDefaultValueAsString) const;
 	static void ReinstanceBatch(TArray<FReinstancingJob>& Reinstancers, TMap< UClass*, UClass* >& InOutOldToNewClassMap, TArray<UObject*>* ObjLoaded);
 	static UClass* FastGenerateSkeletonClass(UBlueprint* BP, FKismetCompilerContext& CompilerContext);
 	static bool IsQueuedForCompilation(UBlueprint* BP);
@@ -112,6 +114,10 @@ struct FBlueprintCompilationManagerImpl : public FGCObject
 	// populated by FlushCompilationQueueImpl, cleared by FlushReinstancingQueueImpl:
 	TMap<UClass*, UClass*> ClassesToReinstance;
 	
+	// Map to old default values, useful for providing access to this data throughout
+	// the compilation process:
+	TMap<UBlueprint*, UObject*> OldCDOs;
+	
 	// Blueprints that should be saved after the compilation pass is complete:
 	TArray<UBlueprint*> CompiledBlueprintsToSave;
 
@@ -191,6 +197,19 @@ void FBlueprintCompilationManagerImpl::CompileSynchronouslyImpl(const FBPCompile
 	FlushCompilationQueueImpl(nullptr, bSuppressBroadcastCompiled, &CompiledBlueprints);
 	FlushReinstancingQueueImpl();
 	
+	if (FBlueprintEditorUtils::IsLevelScriptBlueprint(Request.BPToCompile))
+	{
+		// When the Blueprint is recompiled, then update the bound events for level scripting
+		ULevelScriptBlueprint* LevelScriptBP = CastChecked<ULevelScriptBlueprint>(Request.BPToCompile);
+
+		// ULevel::OnLevelScriptBlueprintChanged needs to be run after the CDO has
+		// been updated as it respawns the actor:
+		if (ULevel* BPLevel = LevelScriptBP->GetLevel())
+		{
+			BPLevel->OnLevelScriptBlueprintChanged(LevelScriptBP);
+		}
+	}
+
 	if ( GEditor )
 	{
 		GEditor->BroadcastBlueprintReinstanced();
@@ -269,14 +288,7 @@ struct FCompilerData
 		InternalOptions.bSaveIntermediateProducts = (UserOptions & EBlueprintCompileOptions::SaveIntermediateProducts ) != EBlueprintCompileOptions::None;
 		InternalOptions.CompileType = bBytecodeOnly ? EKismetCompileType::BytecodeOnly : EKismetCompileType::Full;
 
-		if( UWidgetBlueprint* WidgetBP = Cast<UWidgetBlueprint>(BP))
-		{
-			Compiler = UWidgetBlueprint::GetCompilerForWidgetBP(WidgetBP, *ActiveResultsLog, InternalOptions);
-		}
-		else
-		{
-			Compiler = FKismetCompilerContext::GetCompilerForBP(BP, *ActiveResultsLog, InternalOptions);
-		}
+		Compiler = FBlueprintEditorModule::GetCompiler(BP, *ActiveResultsLog, InternalOptions);
 	}
 
 	bool IsSkeletonOnly() const { return JobType == ECompilationManagerJobType::SkeletonOnly; }
@@ -290,6 +302,7 @@ struct FCompilerData
 	bool ShouldCompileClassLayout() const { return JobType == ECompilationManagerJobType::Normal; }
 	bool ShouldCompileClassFunctions() const { return JobType == ECompilationManagerJobType::Normal; }
 	bool ShouldRelinkAfterSkippingCompile() const { return JobType == ECompilationManagerJobType::RelinkOnly; }
+	bool ShouldSkipIfDependenciesAreUnchanged() const { return InternalOptions.CompileType == EKismetCompileType::BytecodeOnly || JobType == ECompilationManagerJobType::RelinkOnly; }
 
 	UBlueprint* BP;
 	FCompilerResultsLog* ActiveResultsLog;
@@ -307,7 +320,7 @@ struct FReinstancingJob
 	TSharedPtr<FBlueprintCompileReinstancer> Reinstancer;
 	TSharedPtr<FKismetCompilerContext> Compiler;
 };
-	
+
 void FBlueprintCompilationManagerImpl::FlushCompilationQueueImpl(TArray<UObject*>* ObjLoaded, bool bSuppressBroadcastCompiled, TArray<UBlueprint*>* BlueprintsCompiled)
 {
 	TGuardValue<bool> GuardTemplateNameFlag(GCompilingBlueprint, true);
@@ -340,7 +353,18 @@ void FBlueprintCompilationManagerImpl::FlushCompilationQueueImpl(TArray<UObject*
 					if(!IsQueuedForCompilation(DependentBlueprint))
 					{
 						DependentBlueprint->bQueuedForCompilation = true;
-						CurrentlyCompilingBPs.Add(FCompilerData(DependentBlueprint, ECompilationManagerJobType::Normal, nullptr, EBlueprintCompileOptions::None, true));
+						// Because we're adding this as a bytecode only blueprint compile we don't need to 
+						// recursively recompile dependencies. The assumption is that a bytecode only compile
+						// will not change the class layout. @todo: add an ensure to detect class layout changes
+						CurrentlyCompilingBPs.Add(
+							FCompilerData(
+								DependentBlueprint, 
+								ECompilationManagerJobType::Normal, 
+								nullptr, 
+								EBlueprintCompileOptions::None, 
+								true
+							)
+						);
 						BlueprintsToRecompile.Add(DependentBlueprint);
 					}
 				}
@@ -519,30 +543,55 @@ void FBlueprintCompilationManagerImpl::FlushCompilationQueueImpl(TArray<UObject*
 
 		// STAGE VII: safely throw away old skeleton CDOs:
 		{
-			TMap<UClass*, UClass*> OldSkeletonToNewSkeleton;
+			TMap<UClass*, UClass*> NewSkeletonToOldSkeleton;
 			for (FCompilerData& CompilerData : CurrentlyCompilingBPs)
 			{
 				UBlueprint* BP = CompilerData.BP;
 				UClass* OldSkeletonClass = BP->SkeletonGeneratedClass;
 				if(OldSkeletonClass)
 				{
-					MoveSkelCDOAside(OldSkeletonClass, OldSkeletonToNewSkeleton);
+					MoveSkelCDOAside(OldSkeletonClass, NewSkeletonToOldSkeleton);
 				}
 			}
-		}
 		
-		// STAGE VIII: recompile skeleton
-		{
-			TMap<UClass*, UClass*> OldSkeletonToNewSkeleton;
+		
+			// STAGE VIII: recompile skeleton
+
+			// if any function signatures have changed in this skeleton class we will need to recompile all dependencies, but if not
+			// then we can avoid dependency recompilation:
+			TSet<UBlueprint*> BlueprintsWithSignatureChanges;
+			const UBlueprintEditorProjectSettings* EditorProjectSettings = GetDefault<UBlueprintEditorProjectSettings>();
+			bool bSkipUnneededDependencyCompilation = EditorProjectSettings->bSkipUnneededDependencyCompilation;
+
 			for (FCompilerData& CompilerData : CurrentlyCompilingBPs)
 			{
 				UBlueprint* BP = CompilerData.BP;
 		
 				if(CompilerData.ShouldRegenerateSkeleton())
 				{
-					UClass* OldSkeleton = BP->SkeletonGeneratedClass;
 					BP->SkeletonGeneratedClass = FastGenerateSkeletonClass(BP, *(CompilerData.Compiler) );
-					OldSkeletonToNewSkeleton.Add(OldSkeleton, BP->SkeletonGeneratedClass);
+					UBlueprintGeneratedClass* AuthoritativeClass = Cast<UBlueprintGeneratedClass>(BP->GeneratedClass);
+					if(AuthoritativeClass && bSkipUnneededDependencyCompilation)
+					{
+						if(CompilerData.InternalOptions.CompileType == EKismetCompileType::Full )
+						{
+							for (TFieldIterator<UFunction> FuncIt(AuthoritativeClass, EFieldIteratorFlags::ExcludeSuper); FuncIt; ++FuncIt)
+							{
+								// We assume that if the func is FUNC_BlueprintCallable that it will be present in the Skeleton class.
+								// If it is not in the skeleton class we will always think that this blueprints public interface has 
+								// changed. Not a huge deal, but will mean we recompile dependencies more often than necessary.
+								if(FuncIt->HasAnyFunctionFlags(EFunctionFlags::FUNC_BlueprintCallable))
+								{
+									UFunction* NewFunction = BP->SkeletonGeneratedClass->FindFunctionByName((*FuncIt)->GetFName());
+									if(NewFunction == nullptr || !NewFunction->IsSignatureCompatibleWith(*FuncIt))
+									{
+										BlueprintsWithSignatureChanges.Add(BP);
+										break;
+									}
+								}
+							}
+						}
+					}
 				}
 				else
 				{
@@ -552,12 +601,7 @@ void FBlueprintCompilationManagerImpl::FlushCompilationQueueImpl(TArray<UObject*
 
 					// CDO needs to be moved aside already:
 					ensure(SkeletonToRelink->ClassDefaultObject == nullptr);
-
-					UClass * const* OverridenParent = OldSkeletonToNewSkeleton.Find(SkeletonToRelink->GetSuperClass());
-					if(OverridenParent && *OverridenParent)
-					{
-						SkeletonToRelink->SetSuperStruct(*OverridenParent);
-					}
+					ensure(!SkeletonToRelink->GetSuperClass()->HasAnyClassFlags(CLASS_NewerVersionExists));
 
 					SkeletonToRelink->Bind();
 					SkeletonToRelink->StaticLink(true);
@@ -571,6 +615,62 @@ void FBlueprintCompilationManagerImpl::FlushCompilationQueueImpl(TArray<UObject*
 					BP->GeneratedClass->ClearFunctionMapsCaches();
 				}
 			}
+
+			// Skip further compilation for blueprints that are being bytecode compiled as a dependency of something that has
+			// not had a change in its function parameters:
+			auto DependenciesAreCompiled = [&BlueprintsWithSignatureChanges](FCompilerData& Data)
+			{
+				if(Data.ShouldSkipIfDependenciesAreUnchanged())
+				{
+					// if our parent is still being compiled, then we still need to be compiled:
+					UClass* Iter = Data.BP->ParentClass;
+					while(Iter)
+					{
+						if(UBlueprint* BP = Cast<UBlueprint>(Iter->ClassGeneratedBy))
+						{
+							if(BP->bBeingCompiled)
+							{
+								return false;
+							}
+						}
+						Iter = Iter->GetSuperClass();
+					}
+
+					// otherwise if we're dependent on a blueprint that had a function signature change, we still need to be compiled:
+					ensure(Data.BP->bCachedDependenciesUpToDate);
+					ensure(Data.BP->CachedDependencies.Num() > 0); // why are we bytecode compiling a blueprint with no dependencies?
+					for(const TWeakObjectPtr<UBlueprint>& Dependency : Data.BP->CachedDependencies)
+					{
+						if (UBlueprint* DependencyBP = Dependency.Get())
+						{
+							if(DependencyBP->bBeingCompiled && BlueprintsWithSignatureChanges.Contains(DependencyBP))
+							{
+								return false;
+							}
+						}
+					}
+					
+					Data.BP->bBeingCompiled = false;
+					if(UPackage* Package = Data.BP->GetOutermost())
+					{
+						Package->SetDirtyFlag(Data.bPackageWasDirty);
+					}
+					if(Data.ResultsLog)
+					{
+						Data.ResultsLog->EndEvent();
+					}
+					Data.BP->bQueuedForCompilation = false;
+					return true;
+				}
+			
+				return false;
+			};
+
+			if(bSkipUnneededDependencyCompilation)
+			{
+				// Order very much matters, but we could RemoveAllSwap and re-sort:
+				CurrentlyCompilingBPs.RemoveAll(DependenciesAreCompiled);
+			}
 		}
 
 		// STAGE IX: Reconstruct nodes and replace deprecated nodes, then broadcast 'precompile
@@ -641,7 +741,6 @@ void FBlueprintCompilationManagerImpl::FlushCompilationQueueImpl(TArray<UObject*
 		// STAGE X: reinstance every blueprint that is queued, note that this means classes in the hierarchy that are *not* being 
 		// compiled will be parented to REINST versions of the class, so type checks (IsA, etc) involving those types
 		// will be incoherent!
-		TMap<UBlueprint*, UObject*> OldCDOs;
 		{
 			for (FCompilerData& CompilerData : CurrentlyCompilingBPs)
 			{
@@ -675,6 +774,7 @@ void FBlueprintCompilationManagerImpl::FlushCompilationQueueImpl(TArray<UObject*
 			if(BP->GeneratedClass && BP->GeneratedClass->GetSuperClass()->HasAnyClassFlags(CLASS_NewerVersionExists))
 			{
 				BP->GeneratedClass->SetSuperStruct(BP->GeneratedClass->GetSuperClass()->GetAuthoritativeClass());
+
 				if(CompilerData.ShouldResetClassMembers())
 				{
 					BP->GeneratedClass->Children = NULL;
@@ -817,6 +917,8 @@ void FBlueprintCompilationManagerImpl::FlushCompilationQueueImpl(TArray<UObject*
 
 			FScopedDurationTimer ReinstTimer(GTimeReinstancing);
 			ReinstanceBatch(Reinstancers, ClassesToReinstance, ObjLoaded);
+
+			OldCDOs.Empty();
 		}
 		
 		// STAGE XV: CLEAR TEMPORARY FLAGS
@@ -824,20 +926,6 @@ void FBlueprintCompilationManagerImpl::FlushCompilationQueueImpl(TArray<UObject*
 		{
 			UBlueprint* BP = CompilerData.BP;
 			FBlueprintEditorUtils::UpdateDelegatesInBlueprint(BP);
-
-			if (FBlueprintEditorUtils::IsLevelScriptBlueprint(BP))
-			{
-				// When the Blueprint is recompiled, then update the bound events for level scripting
-				ULevelScriptBlueprint* LevelScriptBP = CastChecked<ULevelScriptBlueprint>(BP);
-
-				// ULevel::OnLevelScriptBlueprintChanged needs to be run after the CDO has
-				// been updated as it respawns the actor:
-				if (ULevel* BPLevel = LevelScriptBP->GetLevel())
-				{
-					BPLevel->OnLevelScriptBlueprintChanged(LevelScriptBP);
-				}
-			}
-
 			if(!BP->bIsRegeneratingOnLoad && BP->GeneratedClass)
 			{
 				FKismetEditorUtilities::StripExternalComponents(BP);
@@ -855,8 +943,6 @@ void FBlueprintCompilationManagerImpl::FlushCompilationQueueImpl(TArray<UObject*
 				// refresh each dependent blueprint
 				for (UBlueprint* Dependent : DependentBPs)
 				{
-					// for interface changes, auto-refresh nodes on any dependent blueprints
-					// note: RefreshAllNodes() will internally send a change notification event to the dependent blueprint
 					if(!BP->bIsRegeneratingOnLoad)
 					{
 						// Some logic (e.g. UObject::ProcessInternal) uses this flag to suppress warnings:
@@ -958,6 +1044,7 @@ void FBlueprintCompilationManagerImpl::FlushReinstancingQueueImpl()
 		
 		TGuardValue<bool> ReinstancingGuard(GIsReinstancing, true);
 		FBlueprintCompileReinstancer::BatchReplaceInstancesOfClass(ClassesToReinstance);
+
 		ClassesToReinstance.Empty();
 	}
 	
@@ -982,6 +1069,36 @@ bool FBlueprintCompilationManagerImpl::IsGeneratedClassLayoutReady() const
 	return bGeneratedClassLayoutReady;
 }
 
+void FBlueprintCompilationManagerImpl::GetDefaultValue(const UClass* ForClass, const UProperty* Property, FString& OutDefaultValueAsString) const
+{
+	if(!ForClass || !Property)
+	{
+		return;
+	}
+
+	if (ForClass->ClassDefaultObject)
+	{
+		FBlueprintEditorUtils::PropertyValueToString(Property, (uint8*)ForClass->ClassDefaultObject, OutDefaultValueAsString);
+	}
+	else
+	{
+		UBlueprint* BP = Cast<UBlueprint>(ForClass->ClassGeneratedBy);
+		if(ensure(BP))
+		{
+			const UObject* const* OldCDO = OldCDOs.Find(BP);
+			if(OldCDO && *OldCDO)
+			{
+				const UClass* OldClass = (*OldCDO)->GetClass();
+				const UProperty* OldProperty = OldClass->FindPropertyByName(Property->GetFName());
+				if(OldProperty)
+				{
+					FBlueprintEditorUtils::PropertyValueToString(OldProperty, (uint8*)*OldCDO, OutDefaultValueAsString);
+				}
+			}
+		}
+	}
+}
+
 void FBlueprintCompilationManagerImpl::ReinstanceBatch(TArray<FReinstancingJob>& Reinstancers, TMap< UClass*, UClass* >& InOutOldToNewClassMap, TArray<UObject*>* ObjLoaded)
 {
 	const auto FilterOutOfDateClasses = [](TArray<UClass*>& ClassList)
@@ -1468,10 +1585,10 @@ UClass* FBlueprintCompilationManagerImpl::FastGenerateSkeletonClass(UBlueprint*
 			NewFunction->FunctionFlags |= InFunctionFlags;
 			for(UEdGraphPin* Pin : InputPins)
 			{
-				if(Pin->Direction == EEdGraphPinDirection::EGPD_Output && !Schema->IsExecPin(*Pin) && Pin->ParentPin == nullptr && Pin->GetName() != UK2Node_Event::DelegateOutputName)
+				if(Pin->Direction == EEdGraphPinDirection::EGPD_Output && !Schema->IsExecPin(*Pin) && Pin->ParentPin == nullptr && Pin->GetFName() != UK2Node_Event::DelegateOutputName)
 				{
 					// Reimplementation of FKismetCompilerContext::CreatePropertiesFromList without dependence on 'terms'
-					UProperty* Param = FKismetCompilerUtilities::CreatePropertyOnScope(NewFunction, *(Pin->PinName), Pin->PinType, Ret, CPF_BlueprintVisible|CPF_BlueprintReadOnly, Schema, MessageLog);
+					UProperty* Param = FKismetCompilerUtilities::CreatePropertyOnScope(NewFunction, Pin->PinName, Pin->PinType, Ret, CPF_BlueprintVisible|CPF_BlueprintReadOnly, Schema, MessageLog);
 					if(Param)
 					{
 						Param->SetFlags(RF_Transient);
@@ -1521,7 +1638,7 @@ UClass* FBlueprintCompilationManagerImpl::FastGenerateSkeletonClass(UBlueprint*
 			{
 				// Gather all input pins on these nodes, these are 
 				// the outputs of the function:
-				TSet<FString> UsedPinNames;
+				TSet<FName> UsedPinNames;
 				static const FName RetValName = FName(TEXT("ReturnValue"));
 				for(UK2Node_FunctionResult* Node : ReturnNodes)
 				{
@@ -1533,7 +1650,7 @@ UClass* FBlueprintCompilationManagerImpl::FastGenerateSkeletonClass(UBlueprint*
 							{
 								UsedPinNames.Add(Pin->PinName);
 							
-								UProperty* Param = FKismetCompilerUtilities::CreatePropertyOnScope(NewFunction, *(Pin->PinName), Pin->PinType, Ret, 0, Schema, MessageLog);
+								UProperty* Param = FKismetCompilerUtilities::CreatePropertyOnScope(NewFunction, Pin->PinName, Pin->PinType, Ret, 0, Schema, MessageLog);
 								if(Param)
 								{
 									Param->SetFlags(RF_Transient);
@@ -1604,7 +1721,7 @@ UClass* FBlueprintCompilationManagerImpl::FastGenerateSkeletonClass(UBlueprint*
 					{
 						if( FindField<UObjectProperty>(NewFunction, TEXT("__WorldContext")) == nullptr )
 						{
-							FEdGraphPinType WorldContextPinType(Schema->PC_Object, FString(), UObject::StaticClass(), EPinContainerType::None, false, FEdGraphTerminalType());
+							FEdGraphPinType WorldContextPinType(UEdGraphSchema_K2::PC_Object, NAME_None, UObject::StaticClass(), EPinContainerType::None, false, FEdGraphTerminalType());
 							UProperty* Param = FKismetCompilerUtilities::CreatePropertyOnScope(NewFunction, TEXT("__WorldContext"), WorldContextPinType, Ret, 0, Schema, MessageLog);
 							if(Param)
 							{
@@ -1655,7 +1772,7 @@ UClass* FBlueprintCompilationManagerImpl::FastGenerateSkeletonClass(UBlueprint*
 					(InputPin->PinType.PinCategory != UEdGraphSchema_K2::PC_Interface) && 
 					!InputPin->DefaultValue.IsEmpty() )
 				{
-					NewFunction->SetMetaData(*InputPin->PinName, *InputPin->DefaultValue);
+					NewFunction->SetMetaData(InputPin->PinName, *InputPin->DefaultValue);
 				}
 			}
 
@@ -1968,6 +2085,15 @@ void FBlueprintCompilationManager::FlushCompilationQueue(TArray<UObject*>* ObjLo
 	}
 }
 
+void FBlueprintCompilationManager::FlushCompilationQueueAndReinstance()
+{
+	if(BPCMImpl)
+	{
+		BPCMImpl->FlushCompilationQueueImpl(nullptr, false, nullptr);
+		BPCMImpl->FlushReinstancingQueueImpl();
+	}
+}
+
 void FBlueprintCompilationManager::CompileSynchronously(const FBPCompileRequest& Request)
 {
 	if(BPCMImpl)
@@ -1993,6 +2119,11 @@ void FBlueprintCompilationManager::NotifyBlueprintLoaded(UBlueprint* BPLoaded)
 	BPCMImpl->QueueForCompilation(FBPCompileRequest(BPLoaded, EBlueprintCompileOptions::IsRegeneratingOnLoad, nullptr));
 }
 
+void FBlueprintCompilationManager::QueueForCompilation(UBlueprint* BPLoaded)
+{
+	BPCMImpl->QueueForCompilation(FBPCompileRequest(BPLoaded, EBlueprintCompileOptions::None, nullptr));
+}
+
 bool FBlueprintCompilationManager::IsGeneratedClassLayoutReady()
 {
 	if(!BPCMImpl)
@@ -2003,4 +2134,18 @@ bool FBlueprintCompilationManager::IsGeneratedClassLayoutReady()
 	return BPCMImpl->IsGeneratedClassLayoutReady();
 }
 
+
+bool FBlueprintCompilationManager::GetDefaultValue(const UClass* ForClass, const UProperty* Property, FString& OutDefaultValueAsString)
+{
+	if(!BPCMImpl)
+	{
+		// legacy behavior: can't provide CDO for classes currently being compiled
+		return false;
+	}
+
+	BPCMImpl->GetDefaultValue(ForClass, Property, OutDefaultValueAsString);
+	return true;
+}
+
 #undef LOCTEXT_NAMESPACE
+
