AbilitySystem: Fix Activate not being called on actor-spawning ability tasks. This can cause tasks to leak and other problems in certain use cases.

#UDN 519734
#tests confirm UGameplayAbility::OnGameplayTaskActivated is called for WaitTargetData and other latent, actor spawning ability tasks (non latent/instant tasks will not call this because they are marked pending kill immiediately after spawning the actor).
#rnx

diff --git a/Engine/Source/Editor/GameplayTasksEditor/Private/K2Node_LatentGameplayTaskCall.cpp b/Engine/Source/Editor/GameplayTasksEditor/Private/K2Node_LatentGameplayTaskCall.cpp
--- a/Engine/Source/Editor/GameplayTasksEditor/Private/K2Node_LatentGameplayTaskCall.cpp
+++ b/Engine/Source/Editor/GameplayTasksEditor/Private/K2Node_LatentGameplayTaskCall.cpp
@@ -840,11 +840,57 @@ void UK2Node_LatentGameplayTaskCall::ExpandNode(class FKismetCompilerContext& Co
 
 	LastThenPin = CallPostSpawnnProxyObjectNode->FindPinChecked(UEdGraphSchema_K2::PN_Then);
 
+	// --------------------------------------------------------------------------------------
+	// Create a call to activate the proxy object if necessary
+	// --------------------------------------------------------------------------------------
 
+	UK2Node_IfThenElse* ProxyActivateValidateProxyNode = nullptr;
+
+	if (ProxyActivateFunctionName != NAME_None)
+	{
+		// Validate the proxy object is still valid. Its possible the task ends while calling FinishSpawning, in which case we don't need to call the ProxyActivateFunction.		
+		UK2Node_CallFunction* ProxyActivateIsValidFuncNode = CompilerContext.SpawnIntermediateNode<UK2Node_CallFunction>(this, SourceGraph);
+		ProxyActivateIsValidFuncNode->FunctionReference.SetExternalMember(IsValidFuncName, UKismetSystemLibrary::StaticClass());
+		ProxyActivateIsValidFuncNode->AllocateDefaultPins();
+		UEdGraphPin* ProxyActivateIsValidInputPin = ProxyActivateIsValidFuncNode->FindPinChecked(TEXT("Object"));
+
+		bIsErrorFree &= Schema->TryCreateConnection(ProxyObjectPin, ProxyActivateIsValidInputPin);
+
+		ProxyActivateValidateProxyNode = CompilerContext.SpawnIntermediateNode<UK2Node_IfThenElse>(this, SourceGraph);
+		ProxyActivateValidateProxyNode->AllocateDefaultPins();
+		bIsErrorFree &= Schema->TryCreateConnection(ProxyActivateIsValidFuncNode->GetReturnValuePin(), ProxyActivateValidateProxyNode->GetConditionPin());
+
+		bIsErrorFree &= Schema->TryCreateConnection(LastThenPin, ProxyActivateValidateProxyNode->GetExecPin());
+		LastThenPin = ProxyActivateValidateProxyNode->GetThenPin();
+
+		// Actually cally the Activate function
+		UK2Node_CallFunction* const CallActivateProxyObjectNode = CompilerContext.SpawnIntermediateNode<UK2Node_CallFunction>(this, SourceGraph);
+		CallActivateProxyObjectNode->FunctionReference.SetExternalMember(ProxyActivateFunctionName, ProxyClass);
+		CallActivateProxyObjectNode->AllocateDefaultPins();
+
+		// Hook up the self connection
+		UEdGraphPin* ActivateCallSelfPin = Schema->FindSelfPin(*CallActivateProxyObjectNode, EGPD_Input);
+		check(ActivateCallSelfPin);
+
+		bIsErrorFree &= Schema->TryCreateConnection(ProxyObjectPin, ActivateCallSelfPin);
+
+		// Hook the activate node up in the exec chain
+		UEdGraphPin* ActivateExecPin = CallActivateProxyObjectNode->FindPinChecked(UEdGraphSchema_K2::PN_Execute);
+		UEdGraphPin* ActivateThenPin = CallActivateProxyObjectNode->FindPinChecked(UEdGraphSchema_K2::PN_Then);
+
+		bIsErrorFree &= Schema->TryCreateConnection(LastThenPin, ActivateExecPin);
+
+		LastThenPin = ActivateThenPin;
+	}
+
+	// --------------------------------------------------------------------------------------
 	// Move the connections from the original node then pin to the last internal then pin
+	// --------------------------------------------------------------------------------------
+
 	bIsErrorFree &= CompilerContext.MovePinLinksToIntermediate(*FindPinChecked(UEdGraphSchema_K2::PN_Then), *LastThenPin).CanSafeConnect();
 	bIsErrorFree &= CompilerContext.CopyPinLinksToIntermediate(*LastThenPin, *BranchElsePin).CanSafeConnect();
 	bIsErrorFree &= CompilerContext.CopyPinLinksToIntermediate(*LastThenPin, *ValidateProxyNode->GetElsePin()).CanSafeConnect();
+	bIsErrorFree &= CompilerContext.CopyPinLinksToIntermediate(*LastThenPin, *ProxyActivateValidateProxyNode->GetElsePin()).CanSafeConnect();
 	
 	if (!bIsErrorFree)
 	{
