fix comment typo
#rnx

diff --git a/Engine/Source/Editor/GameplayTasksEditor/Private/K2Node_LatentGameplayTaskCall.cpp b/Engine/Source/Editor/GameplayTasksEditor/Private/K2Node_LatentGameplayTaskCall.cpp
--- a/Engine/Source/Editor/GameplayTasksEditor/Private/K2Node_LatentGameplayTaskCall.cpp
+++ b/Engine/Source/Editor/GameplayTasksEditor/Private/K2Node_LatentGameplayTaskCall.cpp
@@ -863,7 +863,7 @@ void UK2Node_LatentGameplayTaskCall::ExpandNode(class FKismetCompilerContext& Co
 		bIsErrorFree &= Schema->TryCreateConnection(LastThenPin, ProxyActivateValidateProxyNode->GetExecPin());
 		LastThenPin = ProxyActivateValidateProxyNode->GetThenPin();
 
-		// Actually cally the Activate function
+		// Actually call the Activate function
 		UK2Node_CallFunction* const CallActivateProxyObjectNode = CompilerContext.SpawnIntermediateNode<UK2Node_CallFunction>(this, SourceGraph);
 		CallActivateProxyObjectNode->FunctionReference.SetExternalMember(ProxyActivateFunctionName, ProxyClass);
 		CallActivateProxyObjectNode->AllocateDefaultPins();
