Fix static analysis warning

diff --git a/Engine/Source/Editor/GameplayTasksEditor/Private/K2Node_LatentGameplayTaskCall.cpp b/Engine/Source/Editor/GameplayTasksEditor/Private/K2Node_LatentGameplayTaskCall.cpp
--- a/Engine/Source/Editor/GameplayTasksEditor/Private/K2Node_LatentGameplayTaskCall.cpp
+++ b/Engine/Source/Editor/GameplayTasksEditor/Private/K2Node_LatentGameplayTaskCall.cpp
@@ -890,7 +890,10 @@ void UK2Node_LatentGameplayTaskCall::ExpandNode(class FKismetCompilerContext& Co
 	bIsErrorFree &= CompilerContext.MovePinLinksToIntermediate(*FindPinChecked(UEdGraphSchema_K2::PN_Then), *LastThenPin).CanSafeConnect();
 	bIsErrorFree &= CompilerContext.CopyPinLinksToIntermediate(*LastThenPin, *BranchElsePin).CanSafeConnect();
 	bIsErrorFree &= CompilerContext.CopyPinLinksToIntermediate(*LastThenPin, *ValidateProxyNode->GetElsePin()).CanSafeConnect();
-	bIsErrorFree &= CompilerContext.CopyPinLinksToIntermediate(*LastThenPin, *ProxyActivateValidateProxyNode->GetElsePin()).CanSafeConnect();
+	if (ProxyActivateValidateProxyNode)
+	{
+		bIsErrorFree &= CompilerContext.CopyPinLinksToIntermediate(*LastThenPin, *ProxyActivateValidateProxyNode->GetElsePin()).CanSafeConnect();
+	}
 	
 	if (!bIsErrorFree)
 	{
