Link up all delegate pins in LatentGameplayTaskCall nodes, including delegates from super classes

diff --git a/Engine/Source/Editor/GameplayTasksEditor/Private/K2Node_LatentGameplayTaskCall.cpp b/Engine/Source/Editor/GameplayTasksEditor/Private/K2Node_LatentGameplayTaskCall.cpp
--- a/Engine/Source/Editor/GameplayTasksEditor/Private/K2Node_LatentGameplayTaskCall.cpp
+++ b/Engine/Source/Editor/GameplayTasksEditor/Private/K2Node_LatentGameplayTaskCall.cpp
@@ -651,7 +651,7 @@ void UK2Node_LatentGameplayTaskCall::ExpandNode(class FKismetCompilerContext& Co
 	bIsErrorFree &= Schema->TryCreateConnection(LastThenPin, ValidateProxyNode->GetExecPin());
 	LastThenPin = ValidateProxyNode->GetThenPin();
 
-	for (TFieldIterator<FMulticastDelegateProperty> PropertyIt(ProxyClass, EFieldIteratorFlags::ExcludeSuper); PropertyIt && bIsErrorFree; ++PropertyIt)
+	for (TFieldIterator<FMulticastDelegateProperty> PropertyIt(ProxyClass); PropertyIt && bIsErrorFree; ++PropertyIt)
 	{
 		UEdGraphPin* LastActivatedThenPin = nullptr;
 		bIsErrorFree &= FBaseAsyncTaskHelper::HandleDelegateImplementation(*PropertyIt, VariableOutputs, ProxyObjectPin, LastThenPin, LastActivatedThenPin, this, SourceGraph, CompilerContext);
